#TTP 331057 - EDITOR: MESH PAINT MODE: Texture painting feature appears completely broken

#branch UE4

#proj Editor.MeshPaint
#proj Runtime.Engine

#summary Mesh paint editor now works as before.

#change Regressed most of the change from CL 2033028, however bNeedsTwoCopies is set to false by default.
#change MeshPaintEdMode now sets bNeedsTwoCopies on all render targets it makes uses of.

diff --git a/Engine/Source/Editor/MeshPaint/Private/MeshPaintEdMode.cpp b/Engine/Source/Editor/MeshPaint/Private/MeshPaintEdMode.cpp
--- a/Engine/Source/Editor/MeshPaint/Private/MeshPaintEdMode.cpp
+++ b/Engine/Source/Editor/MeshPaint/Private/MeshPaintEdMode.cpp
@@ -1951,13 +1951,13 @@ void FEdModeMeshPaint::StartPaintingTexture( UStaticMeshComponent* InStaticMeshC
 				{
 					TextureData->PaintRenderTargetTexture = NULL;
 					TextureData->PaintRenderTargetTexture = CastChecked<UTextureRenderTarget2D>( StaticConstructObject( UTextureRenderTarget2D::StaticClass(), GetTransientPackage(), NAME_None, RF_Transient ) );
+					TextureData->PaintRenderTargetTexture->bNeedsTwoCopies = true;
 					const bool bForceLinearGamma = true;
 					TextureData->PaintRenderTargetTexture->InitCustomFormat( TextureWidth, TextureHeight, PF_A16B16G16R16, bForceLinearGamma );
 					TextureData->PaintRenderTargetTexture->UpdateResourceImmediate();
 		
 					//Duplicate the texture we are painting and store it in the transient package. This texture is a backup of the data incase we want to revert before commiting.
 					TextureData->PaintingTexture2DDuplicate = (UTexture2D*)StaticDuplicateObject(Texture2D, GetTransientPackage(), *FString::Printf(TEXT("%s_TEMP"), *Texture2D->GetName()));
-
 				}
 				TextureData->PaintRenderTargetTexture->AddressX = Texture2D->AddressX;
 				TextureData->PaintRenderTargetTexture->AddressY = Texture2D->AddressY;
@@ -1974,11 +1974,11 @@ void FEdModeMeshPaint::StartPaintingTexture( UStaticMeshComponent* InStaticMeshC
 					BrushRenderTargetTexture = CastChecked<UTextureRenderTarget2D>( StaticConstructObject( UTextureRenderTarget2D::StaticClass(), GetTransientPackage(), NAME_None, RF_Transient ) );
 					const bool bForceLinearGamma = true;
 					BrushRenderTargetTexture->ClearColor = FLinearColor::Black;
+					BrushRenderTargetTexture->bNeedsTwoCopies = true;
 					BrushRenderTargetTexture->InitCustomFormat( BrushTargetTextureWidth, BrushTargetTextureHeight, PF_A16B16G16R16, bForceLinearGamma );
 					BrushRenderTargetTexture->UpdateResourceImmediate();
 					BrushRenderTargetTexture->AddressX = TextureData->PaintRenderTargetTexture->AddressX;
 					BrushRenderTargetTexture->AddressY = TextureData->PaintRenderTargetTexture->AddressY;
-
 				}
 
 				const bool bEnableSeamPainting = FMeshPaintSettings::Get().bEnableSeamPainting;
@@ -1994,6 +1994,7 @@ void FEdModeMeshPaint::StartPaintingTexture( UStaticMeshComponent* InStaticMeshC
 						BrushMaskRenderTargetTexture = CastChecked<UTextureRenderTarget2D>( StaticConstructObject( UTextureRenderTarget2D::StaticClass(), GetTransientPackage(), NAME_None, RF_Transient ) );
 						const bool bForceLinearGamma = true;
 						BrushMaskRenderTargetTexture->ClearColor = FLinearColor::Black;
+						BrushMaskRenderTargetTexture->bNeedsTwoCopies = true;
 						BrushMaskRenderTargetTexture->InitCustomFormat( BrushTargetTextureWidth, BrushTargetTextureHeight, PF_B8G8R8A8, bForceLinearGamma );
 						BrushMaskRenderTargetTexture->UpdateResourceImmediate();
 						BrushMaskRenderTargetTexture->AddressX = TextureData->PaintRenderTargetTexture->AddressX;
@@ -2009,11 +2010,11 @@ void FEdModeMeshPaint::StartPaintingTexture( UStaticMeshComponent* InStaticMeshC
 						SeamMaskRenderTargetTexture = CastChecked<UTextureRenderTarget2D>( StaticConstructObject( UTextureRenderTarget2D::StaticClass(), GetTransientPackage(), NAME_None, RF_Transient ) );
 						const bool bForceLinearGamma = true;
 						SeamMaskRenderTargetTexture->ClearColor = FLinearColor::Black;
+						SeamMaskRenderTargetTexture->bNeedsTwoCopies = true;
 						SeamMaskRenderTargetTexture->InitCustomFormat( BrushTargetTextureWidth, BrushTargetTextureHeight, PF_B8G8R8A8, bForceLinearGamma );
 						SeamMaskRenderTargetTexture->UpdateResourceImmediate();
 						SeamMaskRenderTargetTexture->AddressX = TextureData->PaintRenderTargetTexture->AddressX;
 						SeamMaskRenderTargetTexture->AddressY = TextureData->PaintRenderTargetTexture->AddressY;
-
 					}
 					bGenerateSeamMask = true;
 				}
diff --git a/Engine/Source/Runtime/Engine/Classes/Engine/TextureRenderTarget.h b/Engine/Source/Runtime/Engine/Classes/Engine/TextureRenderTarget.h
--- a/Engine/Source/Runtime/Engine/Classes/Engine/TextureRenderTarget.h
+++ b/Engine/Source/Runtime/Engine/Classes/Engine/TextureRenderTarget.h
@@ -19,6 +19,9 @@ class UTextureRenderTarget : public UTexture
 	UPROPERTY(EditAnywhere, BlueprintReadWrite, Category=TextureRenderTarget)
 	float TargetGamma;
 
+	/** If true, there will be two copies in memory - one for the texture and one for the render target. If false, they will share memory if possible. This is useful for scene capture textures that are used in the scene. */
+	uint32 bNeedsTwoCopies:1;
+
 	/**
 	 * Render thread: Access the render target resource for this texture target object
 	 * @return pointer to resource or NULL if not initialized
diff --git a/Engine/Source/Runtime/Engine/Private/TextureRenderTarget.cpp b/Engine/Source/Runtime/Engine/Private/TextureRenderTarget.cpp
--- a/Engine/Source/Runtime/Engine/Private/TextureRenderTarget.cpp
+++ b/Engine/Source/Runtime/Engine/Private/TextureRenderTarget.cpp
@@ -16,7 +16,7 @@ UTextureRenderTarget::UTextureRenderTarget(const class FPostConstructInitializeP
 	NeverStream = true;
 	SRGB = true;
 	LODGroup = TEXTUREGROUP_RenderTarget;	
-
+	bNeedsTwoCopies = false;
 #if WITH_EDITORONLY_DATA
 	CompressionNone = true;
 #endif // #if WITH_EDITORONLY_DATA
diff --git a/Engine/Source/Runtime/Engine/Private/TextureRenderTarget2D.cpp b/Engine/Source/Runtime/Engine/Private/TextureRenderTarget2D.cpp
--- a/Engine/Source/Runtime/Engine/Private/TextureRenderTarget2D.cpp
+++ b/Engine/Source/Runtime/Engine/Private/TextureRenderTarget2D.cpp
@@ -319,7 +319,7 @@ void FTextureRenderTarget2DResource::InitDynamicRHI()
 			1,
 			TexCreateFlags,
 			TexCreate_RenderTargetable,
-			false,
+			Owner->bNeedsTwoCopies,
 			RenderTargetTextureRHI,
 			Texture2DRHI
 			);
