#summary 	Disabled ctrl-select when using vertex paint.
#ttp 330992 	EDITOR: Need to be able to lock actor selection while vertex painting.
#detail 	Previous change allowed multi select if you didnt move on mouse up, this removes it completely when mesh painting (still availble via the scene outliner, of course).

diff --git a/Engine/Source/Editor/MeshPaint/Private/MeshPaintEdMode.cpp b/Engine/Source/Editor/MeshPaint/Private/MeshPaintEdMode.cpp
--- a/Engine/Source/Editor/MeshPaint/Private/MeshPaintEdMode.cpp
+++ b/Engine/Source/Editor/MeshPaint/Private/MeshPaintEdMode.cpp
@@ -534,13 +534,6 @@ bool FEdModeMeshPaint::InputKey( FLevelEditorViewportClient* InViewportClient, F
 				//default to assuming this is a paint command
 				bHandled = true;
 
-				// If no other button was pressed && if a first press and we click OFF of an actor and we will let this pass through so multi-select can attempt to handle it 
-				if ((!(bShiftButtonEvent || bAltButtonEvent || bIsOtherMouseButtonEvent)) && ((InKey == EKeys::LeftMouseButton) && ((InEvent == IE_Pressed) || (InEvent == IE_Released)) && (!bAnyPaintAbleActorsUnderCursor)))
-				{
-					bHandled = false;
-					bIsPainting = false;
-				}
-
 				// Allow Ctrl+B to pass through so we can support the finding of a selected static mesh in the content browser.
 				if ( !(bShiftButtonEvent || bAltButtonEvent || bIsOtherMouseButtonEvent) && ( (InKey == EKeys::B) && (InEvent == IE_Pressed) ) )
 				{
