#summary 	Vertex painting now works with Blueprints.
#ttp 317471 	The vertex paint tool does not work with blueprints
#add 	Added vertex color instance data to allow vertex colors to survive Blueprint re-construction.
#change 	Updated mesh paint mode to work with static mesh components, rather than just static mesh actors.

diff --git a/Engine/Source/Editor/MeshPaint/Private/MeshPaintEdMode.cpp b/Engine/Source/Editor/MeshPaint/Private/MeshPaintEdMode.cpp
--- a/Engine/Source/Editor/MeshPaint/Private/MeshPaintEdMode.cpp
+++ b/Engine/Source/Editor/MeshPaint/Private/MeshPaintEdMode.cpp
@@ -97,7 +97,7 @@ FEdModeMeshPaint::FEdModeMeshPaint()
 /** Destructor */
 FEdModeMeshPaint::~FEdModeMeshPaint()
 {
-	CopiedColorsByLOD.Empty();
+	CopiedColorsByComponent.Empty();
 }
 
 /** FGCObject interface */
@@ -961,39 +961,51 @@ void FEdModeMeshPaint::DoPaint( const FVector& InCameraOrigin,
 			{
 				continue;
 			}
-			UStaticMesh* CurStaticMesh = NULL;
-			UStaticMeshComponent* CurStaticMeshComponent = NULL;
-
-			// Is this a static mesh actor?
-			AStaticMeshActor* StaticMeshActor = Cast< AStaticMeshActor >( CurActor );
-			if( StaticMeshActor != NULL &&
-				StaticMeshActor->StaticMeshComponent.IsValid() &&
-				StaticMeshActor->StaticMeshComponent->StaticMesh != NULL )
-			{
-				bCurActorIsValid = true;
-
-				CurStaticMesh = StaticMeshActor->StaticMeshComponent->StaticMesh;
-				CurStaticMeshComponent = StaticMeshActor->StaticMeshComponent;
-
-				// TODO UE4: this used to be the kDOP representation of the mesh.  Need to replace with the PhysX representation
-			}
-
 
-			if ((InPaintAction == EMeshPaintAction::Fill) && (CurStaticMesh))
+			bool bHasStaticMesh = false;
+			TArray<UStaticMeshComponent*> StaticMeshComponents;
+			CurActor->GetComponents<UStaticMeshComponent>(StaticMeshComponents);
+			for(const auto& CurStaticMeshComponent : StaticMeshComponents)
 			{
-				PaintableActors.Add( CurActor );
-				continue;
+				if(CurStaticMeshComponent != NULL)
+				{
+					UStaticMesh* CurStaticMesh = CurStaticMeshComponent->StaticMesh;
+					if(CurStaticMesh != NULL)
+					{
+						bHasStaticMesh = true;
+						break;
+					}
+				}				
 			}
-			else if ((InPaintAction == EMeshPaintAction::PushInstanceColorsToMesh) && CurStaticMesh)
+
+			if(bHasStaticMesh)
 			{
-				PaintableActors.Add( CurActor );
-				continue;
+				if (InPaintAction == EMeshPaintAction::Fill)
+				{
+					PaintableActors.Add( CurActor );
+					continue;
+				}
+				else if (InPaintAction == EMeshPaintAction::PushInstanceColorsToMesh)
+				{
+					PaintableActors.Add( CurActor );
+					continue;
+				}
+
+				ValidSelectedActors.Add( CurActor );
 			}
 
-			// If the actor was a static mesh or dynamic static mesh actor, it's potentially valid for mesh painting.
-			if ( bCurActorIsValid )
+			for(const auto& CurStaticMeshComponent : StaticMeshComponents)
 			{
-				ValidSelectedActors.Add( CurActor );
+				if(CurStaticMeshComponent == NULL)
+				{
+					continue;
+				}
+
+				UStaticMesh* CurStaticMesh = CurStaticMeshComponent->StaticMesh;
+				if(CurStaticMesh == NULL)
+				{
+					continue;
+				}	
 
 				// Get a temporary body setup the has fully detailed collision for the line traces below.
 				TWeakObjectPtr<UBodySetup>* FindBodySetupPtr = StaticMeshToTempBodySetup.Find(CurStaticMesh);
@@ -1130,144 +1142,155 @@ void FEdModeMeshPaint::DoPaint( const FVector& InCameraOrigin,
 	// Iterate over the selected static meshes under the cursor and paint them!
 	for( int32 CurActorIndex = 0; CurActorIndex < PaintableActors.Num(); ++CurActorIndex )
 	{
-		AStaticMeshActor* HitStaticMeshActor = Cast<AStaticMeshActor>( PaintableActors[ CurActorIndex ] );
-		check( HitStaticMeshActor );
+		AActor* HitActor = Cast<AActor>( PaintableActors[ CurActorIndex ] );
+		check( HitActor );
 
-		UStaticMeshComponent* StaticMeshComponent = NULL;
-		if (HitStaticMeshActor)
+		TArray<UStaticMeshComponent*> StaticMeshComponents;
+		HitActor->GetComponents<UStaticMeshComponent>(StaticMeshComponents);
+		check(StaticMeshComponents.Num() > 0);
+		for(const auto& StaticMeshComponent : StaticMeshComponents)
 		{
-			StaticMeshComponent = HitStaticMeshActor->StaticMeshComponent;
-		}
-		
-		UStaticMesh* StaticMesh = StaticMeshComponent->StaticMesh;
+			if(StaticMeshComponent == NULL)
+			{
+				continue;
+			}
+
+			UStaticMesh* StaticMesh = StaticMeshComponent->StaticMesh;
+			if(StaticMesh == NULL)
+			{
+				continue;
+			}
 
-		check( StaticMesh->GetNumLODs() > PaintingMeshLODIndex );
-		FStaticMeshLODResources& LODModel = StaticMeshComponent->StaticMesh->RenderData->LODResources[ PaintingMeshLODIndex ];
+			check( StaticMesh->GetNumLODs() > PaintingMeshLODIndex );
+			FStaticMeshLODResources& LODModel = StaticMeshComponent->StaticMesh->RenderData->LODResources[ PaintingMeshLODIndex ];
 		
-		// Brush properties
-		const float BrushDepth = BrushRadius;	// NOTE: Actually half of the total depth (like a radius)
-		const float BrushFalloffAmount = FMeshPaintSettings::Get().BrushFalloffAmount;
-		const FLinearColor BrushColor = ((InPaintAction == EMeshPaintAction::Paint) || (InPaintAction == EMeshPaintAction::Fill))? FMeshPaintSettings::Get().PaintColor : FMeshPaintSettings::Get().EraseColor;
+			// Brush properties
+			const float BrushDepth = BrushRadius;	// NOTE: Actually half of the total depth (like a radius)
+			const float BrushFalloffAmount = FMeshPaintSettings::Get().BrushFalloffAmount;
+			const FLinearColor BrushColor = ((InPaintAction == EMeshPaintAction::Paint) || (InPaintAction == EMeshPaintAction::Fill))? FMeshPaintSettings::Get().PaintColor : FMeshPaintSettings::Get().EraseColor;
 
-		// NOTE: We square the brush strength to maximize slider precision in the low range
-		const float BrushStrength =
-			FMeshPaintSettings::Get().BrushStrength * FMeshPaintSettings::Get().BrushStrength *
-			InStrengthScale;
+			// NOTE: We square the brush strength to maximize slider precision in the low range
+			const float BrushStrength =
+				FMeshPaintSettings::Get().BrushStrength * FMeshPaintSettings::Get().BrushStrength *
+				InStrengthScale;
 
-		// Display settings
-		const float VisualBiasDistance = 0.15f;
-		const float NormalLineSize( BrushRadius * 0.35f );	// Make the normal line length a function of brush size
-		const FLinearColor NormalLineColor( 0.3f, 1.0f, 0.3f );
-		const FLinearColor BrushCueColor = bIsPainting ? FLinearColor( 1.0f, 1.0f, 0.3f ) : FLinearColor( 0.3f, 1.0f, 0.3f );
-		const FLinearColor InnerBrushCueColor = bIsPainting ? FLinearColor( 0.5f, 0.5f, 0.1f ) : FLinearColor( 0.1f, 0.5f, 0.1f );
+			// Display settings
+			const float VisualBiasDistance = 0.15f;
+			const float NormalLineSize( BrushRadius * 0.35f );	// Make the normal line length a function of brush size
+			const FLinearColor NormalLineColor( 0.3f, 1.0f, 0.3f );
+			const FLinearColor BrushCueColor = bIsPainting ? FLinearColor( 1.0f, 1.0f, 0.3f ) : FLinearColor( 0.3f, 1.0f, 0.3f );
+			const FLinearColor InnerBrushCueColor = bIsPainting ? FLinearColor( 0.5f, 0.5f, 0.1f ) : FLinearColor( 0.1f, 0.5f, 0.1f );
 
-		FVector BrushXAxis, BrushYAxis;
-		BestTraceResult.Normal.FindBestAxisVectors( BrushXAxis, BrushYAxis );
-		const FVector BrushVisualPosition = BestTraceResult.Location + BestTraceResult.Normal * VisualBiasDistance;
+			FVector BrushXAxis, BrushYAxis;
+			BestTraceResult.Normal.FindBestAxisVectors( BrushXAxis, BrushYAxis );
+			const FVector BrushVisualPosition = BestTraceResult.Location + BestTraceResult.Normal * VisualBiasDistance;
 
 
-		// Precache model -> world transform
-		const FMatrix ActorToWorldMatrix = HitStaticMeshActor->ActorToWorld().ToMatrixWithScale();
+			// Precache model -> world transform
+			const FMatrix ComponentToWorldMatrix = StaticMeshComponent->ComponentToWorld.ToMatrixWithScale();
 
 
-		// Compute the camera position in actor space.  We need this later to check for
-		// backfacing triangles.
-		const FVector ActorSpaceCameraPosition( ActorToWorldMatrix.InverseTransformPosition( InCameraOrigin ) );
-		const FVector ActorSpaceBrushPosition( ActorToWorldMatrix.InverseTransformPosition( BestTraceResult.Location ) );
+			// Compute the camera position in actor space.  We need this later to check for
+			// backfacing triangles.
+			const FVector ComponentSpaceCameraPosition( ComponentToWorldMatrix.InverseTransformPosition( InCameraOrigin ) );
+			const FVector ComponentSpaceBrushPosition( ComponentToWorldMatrix.InverseTransformPosition( BestTraceResult.Location ) );
 		
-		// @todo MeshPaint: Input vector doesn't work well with non-uniform scale
-		const float ActorSpaceBrushRadius = ActorToWorldMatrix.InverseTransformVector( FVector( BrushRadius, 0.0f, 0.0f ) ).Size();
-		const float ActorSpaceSquaredBrushRadius = ActorSpaceBrushRadius * ActorSpaceBrushRadius;
+			// @todo MeshPaint: Input vector doesn't work well with non-uniform scale
+			const float ComponentSpaceBrushRadius = ComponentToWorldMatrix.InverseTransformVector( FVector( BrushRadius, 0.0f, 0.0f ) ).Size();
+			const float ComponentSpaceSquaredBrushRadius = ComponentSpaceBrushRadius * ComponentSpaceBrushRadius;
 
 
-		if( PDI != NULL )
-		{
-			// Draw brush circle
-			const int32 NumCircleSides = 64;
-			DrawCircle( PDI, BrushVisualPosition, BrushXAxis, BrushYAxis, BrushCueColor, BrushRadius, NumCircleSides, SDPG_World );
-
-			// Also draw the inner brush radius
-			const float InnerBrushRadius = BrushRadius - BrushFalloffAmount * BrushRadius;
-			DrawCircle( PDI, BrushVisualPosition, BrushXAxis, BrushYAxis, InnerBrushCueColor, InnerBrushRadius, NumCircleSides, SDPG_World );
-
-			// If we just started painting then also draw a little brush effect
-			if( bIsPainting )
+			if( PDI != NULL )
 			{
-				const float EffectDuration = 0.2f;
+				// Draw brush circle
+				const int32 NumCircleSides = 64;
+				DrawCircle( PDI, BrushVisualPosition, BrushXAxis, BrushYAxis, BrushCueColor, BrushRadius, NumCircleSides, SDPG_World );
+
+				// Also draw the inner brush radius
+				const float InnerBrushRadius = BrushRadius - BrushFalloffAmount * BrushRadius;
+				DrawCircle( PDI, BrushVisualPosition, BrushXAxis, BrushYAxis, InnerBrushCueColor, InnerBrushRadius, NumCircleSides, SDPG_World );
 
-				const double CurTime = FPlatformTime::Seconds();
-				const float TimeSinceStartedPainting = (float)( CurTime - PaintingStartTime );
-				if( TimeSinceStartedPainting <= EffectDuration )
+				// If we just started painting then also draw a little brush effect
+				if( bIsPainting )
 				{
-					// Invert the effect if we're currently erasing
-					float EffectAlpha = TimeSinceStartedPainting / EffectDuration;
-					if( InPaintAction == EMeshPaintAction::Erase )
+					const float EffectDuration = 0.2f;
+
+					const double CurTime = FPlatformTime::Seconds();
+					const float TimeSinceStartedPainting = (float)( CurTime - PaintingStartTime );
+					if( TimeSinceStartedPainting <= EffectDuration )
 					{
-						EffectAlpha = 1.0f - EffectAlpha;
-					}
+						// Invert the effect if we're currently erasing
+						float EffectAlpha = TimeSinceStartedPainting / EffectDuration;
+						if( InPaintAction == EMeshPaintAction::Erase )
+						{
+							EffectAlpha = 1.0f - EffectAlpha;
+						}
 
-					const FLinearColor EffectColor( 0.1f + EffectAlpha * 0.4f, 0.1f + EffectAlpha * 0.4f, 0.1f + EffectAlpha * 0.4f );
-					const float EffectRadius = BrushRadius * EffectAlpha * EffectAlpha;	// Squared curve here (looks more interesting)
-					DrawCircle( PDI, BrushVisualPosition, BrushXAxis, BrushYAxis, EffectColor, EffectRadius, NumCircleSides, SDPG_World );
+						const FLinearColor EffectColor( 0.1f + EffectAlpha * 0.4f, 0.1f + EffectAlpha * 0.4f, 0.1f + EffectAlpha * 0.4f );
+						const float EffectRadius = BrushRadius * EffectAlpha * EffectAlpha;	// Squared curve here (looks more interesting)
+						DrawCircle( PDI, BrushVisualPosition, BrushXAxis, BrushYAxis, EffectColor, EffectRadius, NumCircleSides, SDPG_World );
+					}
 				}
-			}
 
-			// Draw trace surface normal
-			const FVector NormalLineEnd( BrushVisualPosition + BestTraceResult.Normal * NormalLineSize );
-			PDI->DrawLine( BrushVisualPosition, NormalLineEnd, NormalLineColor, SDPG_World );
-		}
+				// Draw trace surface normal
+				const FVector NormalLineEnd( BrushVisualPosition + BestTraceResult.Normal * NormalLineSize );
+				PDI->DrawLine( BrushVisualPosition, NormalLineEnd, NormalLineColor, SDPG_World );
+			}
 
 
 
-		// Mesh paint settings
-		FMeshPaintParameters Params;
-		{
-			Params.PaintMode = FMeshPaintSettings::Get().PaintMode;
-			Params.PaintAction = InPaintAction;
-			Params.BrushPosition = BestTraceResult.Location;
-			Params.BrushNormal = BestTraceResult.Normal;
-			Params.BrushColor = BrushColor;
-			Params.SquaredBrushRadius = BrushRadius * BrushRadius;
-			Params.BrushRadialFalloffRange = BrushFalloffAmount * BrushRadius;
-			Params.InnerBrushRadius = BrushRadius - Params.BrushRadialFalloffRange;
-			Params.BrushDepth = BrushDepth;
-			Params.BrushDepthFalloffRange = BrushFalloffAmount * BrushDepth;
-			Params.InnerBrushDepth = BrushDepth - Params.BrushDepthFalloffRange;
-			Params.BrushStrength = BrushStrength;
-			Params.BrushToWorldMatrix = FMatrix( BrushXAxis, BrushYAxis, Params.BrushNormal, Params.BrushPosition );
-			Params.InverseBrushToWorldMatrix = Params.BrushToWorldMatrix.Inverse();
-			Params.bWriteRed = FMeshPaintSettings::Get().bWriteRed;
-			Params.bWriteGreen = FMeshPaintSettings::Get().bWriteGreen;
-			Params.bWriteBlue = FMeshPaintSettings::Get().bWriteBlue;
-			Params.bWriteAlpha = FMeshPaintSettings::Get().bWriteAlpha;
-			Params.TotalWeightCount = FMeshPaintSettings::Get().TotalWeightCount;
-
-			// Select texture weight index based on whether or not we're painting or erasing
+			// Mesh paint settings
+			FMeshPaintParameters Params;
 			{
-				const int32 PaintWeightIndex = 
-					( InPaintAction == EMeshPaintAction::Paint ) ? FMeshPaintSettings::Get().PaintWeightIndex : FMeshPaintSettings::Get().EraseWeightIndex;
+				Params.PaintMode = FMeshPaintSettings::Get().PaintMode;
+				Params.PaintAction = InPaintAction;
+				Params.BrushPosition = BestTraceResult.Location;
+				Params.BrushNormal = BestTraceResult.Normal;
+				Params.BrushColor = BrushColor;
+				Params.SquaredBrushRadius = BrushRadius * BrushRadius;
+				Params.BrushRadialFalloffRange = BrushFalloffAmount * BrushRadius;
+				Params.InnerBrushRadius = BrushRadius - Params.BrushRadialFalloffRange;
+				Params.BrushDepth = BrushDepth;
+				Params.BrushDepthFalloffRange = BrushFalloffAmount * BrushDepth;
+				Params.InnerBrushDepth = BrushDepth - Params.BrushDepthFalloffRange;
+				Params.BrushStrength = BrushStrength;
+				Params.BrushToWorldMatrix = FMatrix( BrushXAxis, BrushYAxis, Params.BrushNormal, Params.BrushPosition );
+				Params.InverseBrushToWorldMatrix = Params.BrushToWorldMatrix.Inverse();
+				Params.bWriteRed = FMeshPaintSettings::Get().bWriteRed;
+				Params.bWriteGreen = FMeshPaintSettings::Get().bWriteGreen;
+				Params.bWriteBlue = FMeshPaintSettings::Get().bWriteBlue;
+				Params.bWriteAlpha = FMeshPaintSettings::Get().bWriteAlpha;
+				Params.TotalWeightCount = FMeshPaintSettings::Get().TotalWeightCount;
+
+				// Select texture weight index based on whether or not we're painting or erasing
+				{
+					const int32 PaintWeightIndex = 
+						( InPaintAction == EMeshPaintAction::Paint ) ? FMeshPaintSettings::Get().PaintWeightIndex : FMeshPaintSettings::Get().EraseWeightIndex;
+
+					// Clamp the weight index to fall within the total weight count
+					Params.PaintWeightIndex = FMath::Clamp( PaintWeightIndex, 0, Params.TotalWeightCount - 1 );
+				}
 
-				// Clamp the weight index to fall within the total weight count
-				Params.PaintWeightIndex = FMath::Clamp( PaintWeightIndex, 0, Params.TotalWeightCount - 1 );
+				// @todo MeshPaint: Ideally we would default to: TexturePaintingStaticMeshComponent->StaticMesh->LightMapCoordinateIndex
+				//		Or we could indicate in the GUI which channel is the light map set (button to set it?)
+				Params.UVChannel = FMeshPaintSettings::Get().UVChannel;
 			}
 
-			// @todo MeshPaint: Ideally we would default to: TexturePaintingStaticMeshComponent->StaticMesh->LightMapCoordinateIndex
-			//		Or we could indicate in the GUI which channel is the light map set (button to set it?)
-			Params.UVChannel = FMeshPaintSettings::Get().UVChannel;
-		}
+			if( FMeshPaintSettings::Get().ResourceType == EMeshPaintResource::VertexColors )
+			{
+				// Painting vertex colors
+				PaintMeshVertices( StaticMeshComponent, Params, bShouldApplyPaint, LODModel, ComponentSpaceCameraPosition, ComponentToWorldMatrix, PDI, VisualBiasDistance);
 
-		if( FMeshPaintSettings::Get().ResourceType == EMeshPaintResource::VertexColors )
-		{
-			// Painting vertex colors
-			PaintMeshVertices( StaticMeshComponent, Params, bShouldApplyPaint, LODModel, ActorSpaceCameraPosition, ActorToWorldMatrix, PDI, VisualBiasDistance);
+			}
+			else
+			{
+				// Painting textures
+				PaintMeshTexture( StaticMeshComponent, Params, bShouldApplyPaint, LODModel, ComponentSpaceCameraPosition, ComponentToWorldMatrix, ComponentSpaceSquaredBrushRadius, ComponentSpaceBrushPosition );
 
+			}		
 		}
-		else
-		{
-			// Painting textures
-			PaintMeshTexture( StaticMeshComponent, Params, bShouldApplyPaint, LODModel, ActorSpaceCameraPosition, ActorToWorldMatrix, ActorSpaceSquaredBrushRadius, ActorSpaceBrushPosition );
+		
 
-		}
 	}
 
 	// Ends an Undo/Redo transaction, but only for Fill or PushInstanceColorsToMesh operations. Multi frame transactions will end when the user stops painting.
@@ -1363,8 +1386,8 @@ void FEdModeMeshPaint::PaintMeshVertices(
 	const FMeshPaintParameters& Params, 
 	const bool bShouldApplyPaint, 
 	FStaticMeshLODResources& LODModel, 
-	const FVector& ActorSpaceCameraPosition, 
-	const FMatrix& ActorToWorldMatrix, 
+	const FVector& ComponentSpaceCameraPosition, 
+	const FMatrix& ComponentToWorldMatrix, 
 	FPrimitiveDrawInterface* PDI, 
 	const float VisualBiasDistance)
 {
@@ -1393,6 +1416,7 @@ void FEdModeMeshPaint::PaintMeshVertices(
 				ComponentReregisterContext.Reset( new FComponentReregisterContext( StaticMeshComponent ) );
 
 				// Mark the mesh component as modified
+				StaticMeshComponent->SetFlags(RF_Transactional);
 				StaticMeshComponent->Modify();
 
 				// Ensure LODData has enough entries in it, free not required.
@@ -1585,7 +1609,7 @@ void FEdModeMeshPaint::PaintMeshVertices(
 
 					// Check to see if the triangle is front facing
 					FVector TriangleNormal = ( TriVertices[ 1 ] - TriVertices[ 0 ] ^ TriVertices[ 2 ] - TriVertices[ 0 ] ).SafeNormal();
-					const float SignedPlaneDist = FVector::PointPlaneDist( ActorSpaceCameraPosition, TriVertices[ 0 ], TriangleNormal );
+					const float SignedPlaneDist = FVector::PointPlaneDist( ComponentSpaceCameraPosition, TriVertices[ 0 ], TriangleNormal );
 					if( !bOnlyFrontFacing || SignedPlaneDist < 0.0f )
 					{
 						FrontFacingVertexIndices[VertexIndices[ 0 ]] = true;
@@ -1600,7 +1624,7 @@ void FEdModeMeshPaint::PaintMeshVertices(
 					// Grab the mesh vertex and transform it to world space
 					const int32 VertexIndex = CurIndexIt.GetIndex();
 					const FVector& ModelSpaceVertexPosition = LODModel.PositionVertexBuffer.VertexPosition( VertexIndex );
-					FVector WorldSpaceVertexPosition = ActorToWorldMatrix.TransformPosition( ModelSpaceVertexPosition );
+					FVector WorldSpaceVertexPosition = ComponentToWorldMatrix.TransformPosition( ModelSpaceVertexPosition );
 
 					FColor OriginalVertexColor = FColor( 255, 255, 255 );
 
@@ -1703,7 +1727,7 @@ void FEdModeMeshPaint::PaintMeshVertices(
 
 
 /** Paints mesh texture */
-void FEdModeMeshPaint::PaintMeshTexture( UStaticMeshComponent* StaticMeshComponent, const FMeshPaintParameters& Params, const bool bShouldApplyPaint, FStaticMeshLODResources& LODModel, const FVector& ActorSpaceCameraPosition, const FMatrix& ActorToWorldMatrix, const float ActorSpaceSquaredBrushRadius, const FVector& ActorSpaceBrushPosition )
+void FEdModeMeshPaint::PaintMeshTexture( UStaticMeshComponent* StaticMeshComponent, const FMeshPaintParameters& Params, const bool bShouldApplyPaint, FStaticMeshLODResources& LODModel, const FVector& ComponentSpaceCameraPosition, const FMatrix& ComponentToWorldMatrix, const float ComponentSpaceSquaredBrushRadius, const FVector& ComponentSpaceBrushPosition )
 {
 	UTexture2D* TargetTexture2D = GetSelectedTexture();
 	
@@ -1732,7 +1756,7 @@ void FEdModeMeshPaint::PaintMeshTexture( UStaticMeshComponent* StaticMeshCompone
 
 		// Use a bit of distance bias to make sure that we get all of the overlapping triangles.  We
 		// definitely don't want our brush to be cut off by a hard triangle edge
-		const float SquaredRadiusBias = ActorSpaceSquaredBrushRadius * 0.025f;
+		const float SquaredRadiusBias = ComponentSpaceSquaredBrushRadius * 0.025f;
 
 		FStaticMeshLODResources& LODModel = StaticMeshComponent->StaticMesh->RenderData->LODResources[PaintingMeshLODIndex];
 		int32 NumSections = LODModel.Sections.Num();
@@ -1816,12 +1840,12 @@ void FEdModeMeshPaint::PaintMeshTexture( UStaticMeshComponent* StaticMeshCompone
 			}
 		}
 
-		for(FMeshTriOctree::TConstElementBoxIterator<> TriIt(*TexturePaintingStaticMeshOctree, FBoxCenterAndExtent(ActorSpaceBrushPosition, FVector(FMath::Sqrt( ActorSpaceSquaredBrushRadius + SquaredRadiusBias )))); TriIt.HasPendingElements(); TriIt.Advance())
+		for(FMeshTriOctree::TConstElementBoxIterator<> TriIt(*TexturePaintingStaticMeshOctree, FBoxCenterAndExtent(ComponentSpaceBrushPosition, FVector(FMath::Sqrt( ComponentSpaceSquaredBrushRadius + SquaredRadiusBias )))); TriIt.HasPendingElements(); TriIt.Advance())
 		{
 			// Check to see if the triangle is front facing
 			FMeshTriangle const& CurrentTri = TriIt.GetCurrentElement();
 			FVector TriangleNormal = ( CurrentTri.Vertices[ 1 ] - CurrentTri.Vertices[ 0 ] ^ CurrentTri.Vertices[ 2 ] - CurrentTri.Vertices[ 0 ] ).SafeNormal();
-			const float SignedPlaneDist = FVector::PointPlaneDist( ActorSpaceCameraPosition, CurrentTri.Vertices[ 0 ], TriangleNormal );
+			const float SignedPlaneDist = FVector::PointPlaneDist( ComponentSpaceCameraPosition, CurrentTri.Vertices[ 0 ], TriangleNormal );
 			if( !bOnlyFrontFacing || SignedPlaneDist < 0.0f )
 			{
 				// At least one triangle vertex was influenced.
@@ -1873,7 +1897,7 @@ void FEdModeMeshPaint::PaintMeshTexture( UStaticMeshComponent* StaticMeshCompone
 
 			if( TexturePaintingStaticMeshComponent != NULL )
 			{
-				PaintTexture( Params, InfluencedTriangles, ActorToWorldMatrix );
+				PaintTexture( Params, InfluencedTriangles, ComponentToWorldMatrix );
 			}
 		}
 	}
@@ -2045,7 +2069,7 @@ void FEdModeMeshPaint::StartPaintingTexture( UStaticMeshComponent* InStaticMeshC
 /** Paints on a texture */
 void FEdModeMeshPaint::PaintTexture( const FMeshPaintParameters& InParams,
 									 const TArray< int32 >& InInfluencedTriangles,
-									 const FMatrix& InActorToWorldMatrix )
+									 const FMatrix& InComponentToWorldMatrix )
 {
 	// We bail early if there are no influenced triangles
 	if( InInfluencedTriangles.Num() <= 0 )
@@ -2162,7 +2186,7 @@ void FEdModeMeshPaint::PaintTexture( const FMeshPaintParameters& InParams,
 		for( int32 TriVertexNum = 0; TriVertexNum < 3; ++TriVertexNum )
 		{																		 
 			const int32 VertexIndex = Indices[ TriIndex * 3 + TriVertexNum ];
-			CurTriangle.TriVertices[ TriVertexNum ] = InActorToWorldMatrix.TransformPosition( LODModel.PositionVertexBuffer.VertexPosition( VertexIndex ) );
+			CurTriangle.TriVertices[ TriVertexNum ] = InComponentToWorldMatrix.TransformPosition( LODModel.PositionVertexBuffer.VertexPosition( VertexIndex ) );
 			CurTriangle.TriUVs[ TriVertexNum ] = LODModel.VertexBuffer.GetVertexUV( VertexIndex, PaintUVCoordinateIndex );
 
 			// Update bounds
@@ -2785,19 +2809,11 @@ TArray<UStaticMeshComponent*> GetValidStaticMeshComponents()
 			continue;
 		}
 
-		// Currently we only support static mesh actors or dynamic static mesh actors
-		AStaticMeshActor* StaticMeshActor = Cast< AStaticMeshActor >( CurActor );
+		TArray<UStaticMeshComponent*> ActorMeshComponents;
+		CurActor->GetComponents<UStaticMeshComponent>( ActorMeshComponents );
+
+		SMComponents.Append( ActorMeshComponents );
 
-		// Grab the StaticMeshComponents from valid static mesh actors
-		if( StaticMeshActor != NULL && StaticMeshActor->StaticMeshComponent.IsValid() )
-		{
-			SMComponents.AddUnique(StaticMeshActor->StaticMeshComponent);
-		}
-		else
-		{
-			// An unsupported actor type is part of the user selection
-			//UE_LOG(LogMeshPaintEdMode, Warning, TEXT("%s is not supported for TEXTure Painting"), *CurActor->GetName() );
-		}
 	}
 
 	return SMComponents;
@@ -2913,50 +2929,48 @@ void FEdModeMeshPaint::UpdateSettingsForStaticMeshComponent( UStaticMeshComponen
 /** FEdMode: Handling SelectActor */
 bool FEdModeMeshPaint::Select( AActor* InActor, bool bInSelected )
 {
-	AStaticMeshActor* StaticMeshActor = Cast< AStaticMeshActor >( InActor );
-
-	UStaticMeshComponent* StaticMeshComponent = NULL;
-	if( StaticMeshActor != NULL )
+	TArray<UStaticMeshComponent*> StaticMeshComponents;
+	InActor->GetComponents<UStaticMeshComponent>(StaticMeshComponents);
+	for(const auto& StaticMeshComponent : StaticMeshComponents)
 	{
-		StaticMeshComponent = StaticMeshActor->StaticMeshComponent;
-	}
-
-	if( StaticMeshComponent )
-	{
-		if( !bInSelected )
+		if( StaticMeshComponent != NULL )
 		{
-			if(FMeshPaintSettings::Get().ResourceType == EMeshPaintResource::Texture)
-			{
-				// When un-selecting a mesh, save it's settings based on the current properties
-				ClearStaticMeshTextureOverrides(StaticMeshComponent);
-				SaveSettingsForActor(InActor);
-			}
-			else if(FMeshPaintSettings::Get().ResourceType == EMeshPaintResource::VertexColors)
+			if( !bInSelected )
 			{
-				//Propagate painting to lower LODs and stop forcing the rendered mesh to LOD0.
-				ApplyVertexColorsToAllLODs(StaticMeshComponent);
-				ClearForcedLOD(StaticMeshComponent);
+				if(FMeshPaintSettings::Get().ResourceType == EMeshPaintResource::Texture)
 				{
-					FComponentReregisterContext ReregisterContext(StaticMeshComponent);
+					// When un-selecting a mesh, save it's settings based on the current properties
+					ClearStaticMeshTextureOverrides(StaticMeshComponent);
+					SaveSettingsForActor(InActor);
+				}
+				else if(FMeshPaintSettings::Get().ResourceType == EMeshPaintResource::VertexColors)
+				{
+					//Propagate painting to lower LODs and stop forcing the rendered mesh to LOD0.
+					ApplyVertexColorsToAllLODs(StaticMeshComponent);
+					ClearForcedLOD(StaticMeshComponent);
+					{
+						FComponentReregisterContext ReregisterContext(StaticMeshComponent);
+					}
 				}
 			}
-		}
-		else
-		{				
-			if(FMeshPaintSettings::Get().ResourceType == EMeshPaintResource::Texture)
-			{
-				SetAllTextureOverrides(StaticMeshComponent);
-			}
-			else if(FMeshPaintSettings::Get().ResourceType == EMeshPaintResource::VertexColors)
-			{
-				//Painting is done on LOD0 so force the mesh to render only LOD0.
-				ForceBestLOD(StaticMeshComponent);
+			else
+			{				
+				if(FMeshPaintSettings::Get().ResourceType == EMeshPaintResource::Texture)
 				{
-					FComponentReregisterContext ReregisterContext(StaticMeshComponent);
+					SetAllTextureOverrides(StaticMeshComponent);
+				}
+				else if(FMeshPaintSettings::Get().ResourceType == EMeshPaintResource::VertexColors)
+				{
+					//Painting is done on LOD0 so force the mesh to render only LOD0.
+					ForceBestLOD(StaticMeshComponent);
+					{
+						FComponentReregisterContext ReregisterContext(StaticMeshComponent);
+					}
 				}
 			}
-		}
+		}	
 	}
+	
 	return false;
 }
 
@@ -3556,19 +3570,21 @@ EMeshPaintAction::Type FEdModeMeshPaint::GetPaintAction(FViewport* InViewport)
 
 }
 
-/** Removes vertex colors associated with the object (if it's a StaticMeshActor) */
+/** Removes vertex colors associated with the object */
 void FEdModeMeshPaint::RemoveInstanceVertexColors(UObject* Obj) const
 {
-	UStaticMeshComponent* StaticMeshComponent = NULL;
-	AStaticMeshActor* StaticMeshActor = Cast< AStaticMeshActor >( Obj );
-	if( StaticMeshActor != NULL )
-	{
-		StaticMeshComponent = StaticMeshActor->StaticMeshComponent;
-	}
-
-	if(StaticMeshComponent != NULL)
+	AActor* Actor = Cast<AActor>(Obj);
+	if(Actor != NULL)
 	{
-		RemoveComponentInstanceVertexColors(StaticMeshComponent);
+		TArray<UStaticMeshComponent*> StaticMeshComponents;
+		Actor->GetComponents<UStaticMeshComponent>(StaticMeshComponents);
+		for(const auto& StaticMeshComponent : StaticMeshComponents)
+		{
+			if(StaticMeshComponent != NULL)
+			{
+				RemoveComponentInstanceVertexColors(StaticMeshComponent);
+			}
+		}
 	}
 }
 
@@ -3635,10 +3651,12 @@ void FEdModeMeshPaint::RemoveInstanceVertexColorsWorker(UStaticMeshComponent *St
 	StaticMeshComponent->StaticMeshDerivedDataKey = StaticMeshComponent->StaticMesh->RenderData->DerivedDataKey;
 }
 
+
+
 /** Copies vertex colors associated with the currently selected mesh */
 void FEdModeMeshPaint::CopyInstanceVertexColors()
 {
-	CopiedColorsByLOD.Empty();
+	CopiedColorsByComponent.Empty();
 
 	USelection& SelectedActors = *GEditor->GetSelectedActors();
 	if( SelectedActors.Num() != 1 )
@@ -3646,73 +3664,76 @@ void FEdModeMeshPaint::CopyInstanceVertexColors()
 		// warning - works only with 1 actor selected..!
 	}
 	else
-	{
-		UStaticMeshComponent* StaticMeshComponent = NULL;
-		UObject *CurrentObject = SelectedActors.GetSelectedObject( 0 );
-		AStaticMeshActor* StaticMeshActor = Cast< AStaticMeshActor >( CurrentObject );
-		if( StaticMeshActor != NULL )
-		{
-			StaticMeshComponent = StaticMeshActor->StaticMeshComponent;
-		}
-
-		if( StaticMeshComponent )
+	{		
+		AActor* SelectedActor = Cast<AActor>(SelectedActors.GetSelectedObject(0));
+		if( SelectedActor != NULL )
 		{
-			int32 NumLODs = StaticMeshComponent->StaticMesh->GetNumLODs();
-			for( int32 CurLODIndex = 0; CurLODIndex < NumLODs; ++CurLODIndex )
+			TArray<UStaticMeshComponent*> StaticMeshComponents;
+			SelectedActor->GetComponents<UStaticMeshComponent>(StaticMeshComponents);
+			for(const auto& StaticMeshComponent : StaticMeshComponents)
 			{
-				FPerLODVertexColorData& LodColorData = CopiedColorsByLOD[CopiedColorsByLOD.AddZeroed()];
-
-				UStaticMesh* StaticMesh = StaticMeshComponent->StaticMesh;
-				FStaticMeshLODResources& LODModel = StaticMeshComponent->StaticMesh->RenderData->LODResources[ CurLODIndex ];
-				FColorVertexBuffer* ColBuffer = &LODModel.ColorVertexBuffer;
-
-				FPositionVertexBuffer* PosBuffer = &LODModel.PositionVertexBuffer;
-
-				// Is there an override buffer? If so, copy colors from there instead...
-				if( StaticMeshComponent->LODData.Num() > CurLODIndex )
+				if( StaticMeshComponent )
 				{
-					FStaticMeshComponentLODInfo& ComponentLODInfo = StaticMeshComponent->LODData[ CurLODIndex ];
-					if( ComponentLODInfo.OverrideVertexColors )
+					FPerComponentVertexColorData& PerComponentData = CopiedColorsByComponent[CopiedColorsByComponent.Add(FPerComponentVertexColorData(StaticMeshComponent->StaticMesh, StaticMeshComponent->GetSerializedComponentIndex()))];
+
+					int32 NumLODs = StaticMeshComponent->StaticMesh->GetNumLODs();
+					for( int32 CurLODIndex = 0; CurLODIndex < NumLODs; ++CurLODIndex )
 					{
-						ColBuffer = ComponentLODInfo.OverrideVertexColors;
-					}
-				}
+						FPerLODVertexColorData& LodColorData =  PerComponentData.PerLODVertexColorData[PerComponentData.PerLODVertexColorData.AddZeroed()];
 
-				// Copy the colour buffer
-				if( ColBuffer && PosBuffer )
-				{
-					uint32 NumColVertices = ColBuffer->GetNumVertices();
-					uint32 NumPosVertices = PosBuffer->GetNumVertices();
+						UStaticMesh* StaticMesh = StaticMeshComponent->StaticMesh;
+						FStaticMeshLODResources& LODModel = StaticMeshComponent->StaticMesh->RenderData->LODResources[ CurLODIndex ];
+						FColorVertexBuffer* ColBuffer = &LODModel.ColorVertexBuffer;
 
-					if (NumColVertices == NumPosVertices)
-					{
-						// valid color buffer matching the pos verts
-						for( uint32 VertexIndex = 0; VertexIndex < NumColVertices; VertexIndex++ )
+						FPositionVertexBuffer* PosBuffer = &LODModel.PositionVertexBuffer;
+
+						// Is there an override buffer? If so, copy colors from there instead...
+						if( StaticMeshComponent->LODData.Num() > CurLODIndex )
 						{
-							LodColorData.ColorsByIndex.Add( ColBuffer->VertexColor( VertexIndex ) );
-							LodColorData.ColorsByPosition.Add( PosBuffer->VertexPosition( VertexIndex ), ColBuffer->VertexColor( VertexIndex ) );
-						}						
-					}
-					else
-					{
-						// mismatched or empty color buffer - just use white
-						for( uint32 VertexIndex = 0; VertexIndex < NumPosVertices; VertexIndex++ )
+							FStaticMeshComponentLODInfo& ComponentLODInfo = StaticMeshComponent->LODData[ CurLODIndex ];
+							if( ComponentLODInfo.OverrideVertexColors )
+							{
+								ColBuffer = ComponentLODInfo.OverrideVertexColors;
+							}
+						}
+
+						// Copy the colour buffer
+						if( ColBuffer && PosBuffer )
 						{
-							LodColorData.ColorsByIndex.Add( FColor(255,255,255,255) );
-							LodColorData.ColorsByPosition.Add( PosBuffer->VertexPosition( VertexIndex ), FColor(255,255,255,255) );
+							uint32 NumColVertices = ColBuffer->GetNumVertices();
+							uint32 NumPosVertices = PosBuffer->GetNumVertices();
+
+							if (NumColVertices == NumPosVertices)
+							{
+								// valid color buffer matching the pos verts
+								for( uint32 VertexIndex = 0; VertexIndex < NumColVertices; VertexIndex++ )
+								{
+									LodColorData.ColorsByIndex.Add( ColBuffer->VertexColor( VertexIndex ) );
+									LodColorData.ColorsByPosition.Add( PosBuffer->VertexPosition( VertexIndex ), ColBuffer->VertexColor( VertexIndex ) );
+								}						
+							}
+							else
+							{
+								// mismatched or empty color buffer - just use white
+								for( uint32 VertexIndex = 0; VertexIndex < NumPosVertices; VertexIndex++ )
+								{
+									LodColorData.ColorsByIndex.Add( FColor(255,255,255,255) );
+									LodColorData.ColorsByPosition.Add( PosBuffer->VertexPosition( VertexIndex ), FColor(255,255,255,255) );
+								}
+							}
 						}
 					}
 				}
 			}
-		}
+		}	
 	}
 }
 
 /** Pastes vertex colors to the currently selected mesh */
 void FEdModeMeshPaint::PasteInstanceVertexColors()
 {
-	int32 NumLodsInCopyBuffer = CopiedColorsByLOD.Num();
-	if (0 == NumLodsInCopyBuffer)
+	const int32 NumComponentsInCopyBuffer = CopiedColorsByComponent.Num();
+	if(0 == NumComponentsInCopyBuffer)
 	{
 		return;
 	}
@@ -3723,100 +3744,117 @@ void FEdModeMeshPaint::PasteInstanceVertexColors()
 
 	TScopedPointer< FComponentReregisterContext > ComponentReregisterContext;
 
-	for( int32 iMesh = 0; iMesh < SelectedActors.Num(); iMesh++ )
+	for( int32 ActorIndex = 0; ActorIndex < SelectedActors.Num(); ActorIndex++ )
 	{
-		UStaticMeshComponent* StaticMeshComponent = NULL;
-		UObject* CurrentObject = SelectedActors.GetSelectedObject( iMesh );
+		UObject* CurrentObject = SelectedActors.GetSelectedObject( ActorIndex );
 		AActor* CurrentActor = Cast< AActor >( CurrentObject );
-		AStaticMeshActor* StaticMeshActor = Cast< AStaticMeshActor >( CurrentObject );
-		if( StaticMeshActor != NULL )
+		if( CurrentActor != NULL )
 		{
-			CurrentActor = StaticMeshActor;
-			StaticMeshComponent = StaticMeshActor->StaticMeshComponent;
-		}
-
-		if( StaticMeshComponent )
-		{
-			int32 NumLods = StaticMeshComponent->StaticMesh->GetNumLODs();
-			if (0 == NumLods)
+			TArray<UStaticMeshComponent*> StaticMeshComponents;
+			CurrentActor->GetComponents<UStaticMeshComponent>(StaticMeshComponents);
+			for(const auto& StaticMeshComponent : StaticMeshComponents)
 			{
-				continue;
-			}
-
-			ComponentReregisterContext.Reset( new FComponentReregisterContext( StaticMeshComponent ) );
-			StaticMeshComponent->Modify();
-			StaticMeshComponent->SetLODDataCount( NumLods, NumLods );
-			RemoveInstanceVertexColors( CurrentActor );
-
-			for( int32 CurLODIndex = 0; CurLODIndex < NumLods; ++CurLODIndex )
-			{
-				FStaticMeshLODResources& LodRenderData = StaticMeshComponent->StaticMesh->RenderData->LODResources[CurLODIndex];
-				FStaticMeshComponentLODInfo& ComponentLodInfo = StaticMeshComponent->LODData[CurLODIndex];
-
-				TArray< FColor > ReOrderedColors;
-				TArray< FColor >* PasteFromBufferPtr = &ReOrderedColors;
-
-				if (CurLODIndex >= NumLodsInCopyBuffer)
+				if( StaticMeshComponent )
 				{
-					// no corresponding LOD in color paste buffer CopiedColorsByLOD
-					// create array of all white verts
-					ReOrderedColors.AddUninitialized(LodRenderData.GetNumVertices());
-
-					for (int32 TargetVertIdx = 0; TargetVertIdx < LodRenderData.GetNumVertices(); TargetVertIdx++)
+					int32 NumLods = StaticMeshComponent->StaticMesh->GetNumLODs();
+					if (0 == NumLods)
 					{
-						ReOrderedColors[TargetVertIdx] = FColor(255,255,255,255);
+						continue;
 					}
-				}
-				else if (LodRenderData.GetNumVertices() == CopiedColorsByLOD[CurLODIndex].ColorsByIndex.Num())
-				{
-					// verts counts match - copy from color array by index
-					PasteFromBufferPtr = &(CopiedColorsByLOD[CurLODIndex].ColorsByIndex);
-				}
-				else
-				{
-					// verts counts mismatch - build translation/fixup list of colors in ReOrderedColors
-					ReOrderedColors.AddUninitialized(LodRenderData.GetNumVertices());
 
-					// make ReOrderedColors contain one FColor for each vertex in the target mesh
-					// matching the position of the target's vert to the position values in LodColorData.ColorsByPosition
-					for (int32 TargetVertIdx = 0; TargetVertIdx < LodRenderData.GetNumVertices(); TargetVertIdx++)
+					// attempt to find a matching component in our clipboard data
+					const int32 SerializedComponentIndex = StaticMeshComponent->GetSerializedComponentIndex();
+					FPerComponentVertexColorData* FoundColors = NULL;
+					for(auto& CopiedColors : CopiedColorsByComponent)
 					{
-						const FColor* FoundColor =
-							CopiedColorsByLOD[CurLODIndex].ColorsByPosition.Find(LodRenderData.PositionVertexBuffer.VertexPosition(TargetVertIdx));
-
-						if (FoundColor)
+						if(CopiedColors.OriginalMesh.Get() == StaticMeshComponent->StaticMesh &&
+							CopiedColors.ComponentIndex == SerializedComponentIndex)
 						{
-							// A matching color for this vertex was found
-							ReOrderedColors[TargetVertIdx] = *FoundColor;
+							FoundColors = &CopiedColors;
+							break;
 						}
-						else
+					}
+
+					if(FoundColors != NULL)
+					{
+						ComponentReregisterContext.Reset( new FComponentReregisterContext( StaticMeshComponent ) );
+						StaticMeshComponent->SetFlags(RF_Transactional);
+						StaticMeshComponent->Modify();
+						StaticMeshComponent->SetLODDataCount( NumLods, NumLods );
+						RemoveComponentInstanceVertexColors( StaticMeshComponent );
+
+						for( int32 CurLODIndex = 0; CurLODIndex < NumLods; ++CurLODIndex )
 						{
-							// A matching color for this vertex could not be found. Make this vertex white
-							ReOrderedColors[TargetVertIdx] = FColor(255,255,255,255);
+							FStaticMeshLODResources& LodRenderData = StaticMeshComponent->StaticMesh->RenderData->LODResources[CurLODIndex];
+							FStaticMeshComponentLODInfo& ComponentLodInfo = StaticMeshComponent->LODData[CurLODIndex];
+
+							TArray< FColor > ReOrderedColors;
+							TArray< FColor >* PasteFromBufferPtr = &ReOrderedColors;
+
+							const int32 NumLodsInCopyBuffer = FoundColors->PerLODVertexColorData.Num();
+							if (CurLODIndex >= NumLodsInCopyBuffer)
+							{
+								// no corresponding LOD in color paste buffer CopiedColorsByLOD
+								// create array of all white verts
+								ReOrderedColors.AddUninitialized(LodRenderData.GetNumVertices());
+
+								for (int32 TargetVertIdx = 0; TargetVertIdx < LodRenderData.GetNumVertices(); TargetVertIdx++)
+								{
+									ReOrderedColors[TargetVertIdx] = FColor(255,255,255,255);
+								}
+							}
+							else if (LodRenderData.GetNumVertices() == FoundColors->PerLODVertexColorData[CurLODIndex].ColorsByIndex.Num())
+							{
+								// verts counts match - copy from color array by index
+								PasteFromBufferPtr = &(FoundColors->PerLODVertexColorData[CurLODIndex].ColorsByIndex);
+							}
+							else
+							{
+								// verts counts mismatch - build translation/fixup list of colors in ReOrderedColors
+								ReOrderedColors.AddUninitialized(LodRenderData.GetNumVertices());
+
+								// make ReOrderedColors contain one FColor for each vertex in the target mesh
+								// matching the position of the target's vert to the position values in LodColorData.ColorsByPosition
+								for (int32 TargetVertIdx = 0; TargetVertIdx < LodRenderData.GetNumVertices(); TargetVertIdx++)
+								{
+									const FColor* FoundColor =
+										 FoundColors->PerLODVertexColorData[CurLODIndex].ColorsByPosition.Find(LodRenderData.PositionVertexBuffer.VertexPosition(TargetVertIdx));
+
+									if (FoundColor)
+									{
+										// A matching color for this vertex was found
+										ReOrderedColors[TargetVertIdx] = *FoundColor;
+									}
+									else
+									{
+										// A matching color for this vertex could not be found. Make this vertex white
+										ReOrderedColors[TargetVertIdx] = FColor(255,255,255,255);
+									}
+								}
+							}
+
+							if( ComponentLodInfo.OverrideVertexColors )
+							{
+								ComponentLodInfo.ReleaseOverrideVertexColorsAndBlock();
+							}
+							if( ComponentLodInfo.OverrideVertexColors )
+							{
+								ComponentLodInfo.BeginReleaseOverrideVertexColors();
+								FlushRenderingCommands();
+							}
+							else
+							{
+								ComponentLodInfo.OverrideVertexColors = new FColorVertexBuffer;
+								ComponentLodInfo.OverrideVertexColors->InitFromColorArray( *PasteFromBufferPtr );
+							}
+							BeginInitResource( ComponentLodInfo.OverrideVertexColors );				
 						}
-					}
-				}
 
-				if( ComponentLodInfo.OverrideVertexColors )
-				{
-					ComponentLodInfo.ReleaseOverrideVertexColorsAndBlock();
-				}
-				if( ComponentLodInfo.OverrideVertexColors )
-				{
-					ComponentLodInfo.BeginReleaseOverrideVertexColors();
-					FlushRenderingCommands();
-				}
-				else
-				{
-					ComponentLodInfo.OverrideVertexColors = new FColorVertexBuffer;
-					ComponentLodInfo.OverrideVertexColors->InitFromColorArray( *PasteFromBufferPtr );
+						StaticMeshComponent->CachePaintedDataIfNecessary();
+						StaticMeshComponent->StaticMeshDerivedDataKey = StaticMeshComponent->StaticMesh->RenderData->DerivedDataKey;
+					}
 				}
-				BeginInitResource( ComponentLodInfo.OverrideVertexColors );				
 			}
-
-			StaticMeshComponent->CachePaintedDataIfNecessary();
-			StaticMeshComponent->StaticMeshDerivedDataKey = StaticMeshComponent->StaticMesh->RenderData->DerivedDataKey;
-
 		}
 	}
 }
@@ -3830,19 +3868,21 @@ bool FEdModeMeshPaint::RequiresInstanceVertexColorsFixup() const
 	USelection& SelectedActors = *GEditor->GetSelectedActors();
 	for( int32 CurSelectedActorIndex = 0; CurSelectedActorIndex < SelectedActors.Num(); ++CurSelectedActorIndex )
 	{
-		UStaticMeshComponent* StaticMeshComponent = NULL;
-		AStaticMeshActor* StaticMeshActor = Cast< AStaticMeshActor >( SelectedActors.GetSelectedObject( CurSelectedActorIndex ) );
-		if( StaticMeshActor )
+		AActor* SelectedActor = Cast< AActor >( SelectedActors.GetSelectedObject( CurSelectedActorIndex ) );
+		if(SelectedActor != NULL)
 		{
-			StaticMeshComponent = StaticMeshActor->StaticMeshComponent;
-		}
-
-		// If a static mesh component was found and it requires fixup, exit out and indicate as such
-		TArray<int32> LODsToFixup;
-		if( StaticMeshComponent && StaticMeshComponent->RequiresOverrideVertexColorsFixup( LODsToFixup ) )
-		{
-			bRequiresFixup = true;
-			break;
+			TArray<UStaticMeshComponent*> StaticMeshComponents;
+			SelectedActor->GetComponents<UStaticMeshComponent>(StaticMeshComponents);
+			for(const auto& StaticMeshComponent : StaticMeshComponents)
+			{
+				// If a static mesh component was found and it requires fixup, exit out and indicate as such
+				TArray<int32> LODsToFixup;
+				if( StaticMeshComponent && StaticMeshComponent->RequiresOverrideVertexColorsFixup( LODsToFixup ) )
+				{
+					bRequiresFixup = true;
+					break;
+				}
+			}
 		}
 	}
 
@@ -3856,17 +3896,19 @@ void FEdModeMeshPaint::FixupInstanceVertexColors() const
 	USelection& SelectedActors = *GEditor->GetSelectedActors();
 	for( int32 CurSelectedActorIndex = 0; CurSelectedActorIndex < SelectedActors.Num(); ++CurSelectedActorIndex )
 	{
-		UStaticMeshComponent* StaticMeshComponent = NULL;
-		AStaticMeshActor* StaticMeshActor = Cast< AStaticMeshActor >( SelectedActors.GetSelectedObject( CurSelectedActorIndex ) );
-		if( StaticMeshActor != NULL )
-		{
-			StaticMeshComponent = StaticMeshActor->StaticMeshComponent;
-		}
-
-		// If a static mesh component was found, attempt to fixup its override colors
-		if( StaticMeshComponent )
+		AActor* SelectedActor = Cast< AActor >( SelectedActors.GetSelectedObject( CurSelectedActorIndex ) );
+		if(SelectedActor != NULL)
 		{
-			StaticMeshComponent->FixupOverrideColorsIfNecessary();
+			TArray<UStaticMeshComponent*> StaticMeshComponents;
+			SelectedActor->GetComponents<UStaticMeshComponent>(StaticMeshComponents);
+			for(const auto& StaticMeshComponent : StaticMeshComponents)
+			{
+				// If a static mesh component was found, attempt to fixup its override colors
+				if(StaticMeshComponent != NULL)
+				{
+					StaticMeshComponent->FixupOverrideColorsIfNecessary();
+				}
+			}
 		}
 	}
 }
@@ -3876,10 +3918,15 @@ void FEdModeMeshPaint::ForceBestLOD()
 	USelection& SelectedActors = *GEditor->GetSelectedActors();
 	for( int32 CurSelectedActorIndex = 0; CurSelectedActorIndex < SelectedActors.Num(); ++CurSelectedActorIndex )
 	{
-		AStaticMeshActor* StaticMeshActor = Cast< AStaticMeshActor >( SelectedActors.GetSelectedObject( CurSelectedActorIndex ) );
-		if( StaticMeshActor != NULL )
+		AActor* SelectedActor = Cast< AActor >( SelectedActors.GetSelectedObject( CurSelectedActorIndex ) );
+		if(SelectedActor != NULL)
 		{
-			ForceBestLOD(StaticMeshActor->StaticMeshComponent);
+			TArray<UStaticMeshComponent*> StaticMeshComponents;
+			SelectedActor->GetComponents<UStaticMeshComponent>(StaticMeshComponents);
+			for(const auto& StaticMeshComponent : StaticMeshComponents)
+			{
+				ForceBestLOD(StaticMeshComponent);
+			}
 		}
 	}
 }
@@ -3899,10 +3946,15 @@ void FEdModeMeshPaint::ClearForcedLOD()
 	USelection& SelectedActors = *GEditor->GetSelectedActors();
 	for( int32 CurSelectedActorIndex = 0; CurSelectedActorIndex < SelectedActors.Num(); ++CurSelectedActorIndex )
 	{
-		AStaticMeshActor* StaticMeshActor = Cast< AStaticMeshActor >( SelectedActors.GetSelectedObject( CurSelectedActorIndex ) );
-		if( StaticMeshActor != NULL )
+		AActor* SelectedActor = Cast< AActor >( SelectedActors.GetSelectedObject( CurSelectedActorIndex ) );
+		if(SelectedActor != NULL)
 		{
-			ClearForcedLOD(StaticMeshActor->StaticMeshComponent);
+			TArray<UStaticMeshComponent*> StaticMeshComponents;
+			SelectedActor->GetComponents<UStaticMeshComponent>(StaticMeshComponents);
+			for(const auto& StaticMeshComponent : StaticMeshComponents)
+			{
+				ClearForcedLOD(StaticMeshComponent);
+			}
 		}
 	}
 }
@@ -3923,10 +3975,15 @@ void FEdModeMeshPaint::ApplyVertexColorsToAllLODs()
 	USelection& SelectedActors = *GEditor->GetSelectedActors();
 	for( int32 CurSelectedActorIndex = 0; CurSelectedActorIndex < SelectedActors.Num(); ++CurSelectedActorIndex )
 	{
-		AStaticMeshActor* StaticMeshActor = Cast< AStaticMeshActor >( SelectedActors.GetSelectedObject( CurSelectedActorIndex ) );
-		if( StaticMeshActor != NULL )
+		AActor* SelectedActor = Cast< AActor >( SelectedActors.GetSelectedObject( CurSelectedActorIndex ) );
+		if(SelectedActor != NULL)
 		{
-			ApplyVertexColorsToAllLODs(StaticMeshActor->StaticMeshComponent);
+			TArray<UStaticMeshComponent*> StaticMeshComponents;
+			SelectedActor->GetComponents<UStaticMeshComponent>(StaticMeshComponents);
+			for(const auto& StaticMeshComponent : StaticMeshComponents)
+			{
+				ApplyVertexColorsToAllLODs(StaticMeshComponent);
+			}
 		}
 	}
 }
@@ -4050,17 +4107,18 @@ void FEdModeMeshPaint::RemoveInstanceMaterialAndTexture() const
 	USelection& SelectedActors = *GEditor->GetSelectedActors();
 	for( int32 CurSelectedActorIndex = 0; CurSelectedActorIndex < SelectedActors.Num(); ++CurSelectedActorIndex )
 	{
-		AActor* CurSelectedActor = CastChecked<AActor>( SelectedActors.GetSelectedObject( CurSelectedActorIndex ) );
-		UStaticMeshComponent* StaticMeshComponent = NULL;
-		AStaticMeshActor* StaticMeshActor = Cast< AStaticMeshActor >( CurSelectedActor );
-		if( StaticMeshActor != NULL )
-		{
-			StaticMeshComponent = StaticMeshActor->StaticMeshComponent;
-		}
-
-		if( StaticMeshComponent != NULL )
+		AActor* SelectedActor = Cast< AActor >( SelectedActors.GetSelectedObject( CurSelectedActorIndex ) );
+		if(SelectedActor != NULL)
 		{
-			// ...
+			TArray<UStaticMeshComponent*> StaticMeshComponents;
+			SelectedActor->GetComponents<UStaticMeshComponent>(StaticMeshComponents);
+			for(const auto& StaticMeshComponent : StaticMeshComponents)
+			{
+				if( StaticMeshComponent != NULL )
+				{
+					// @todo: this function
+				}
+			}
 		}
 	}
 }
@@ -4079,31 +4137,32 @@ bool FEdModeMeshPaint::GetSelectedMeshInfo( int32& OutTotalBaseVertexColorBytes,
 	USelection& SelectedActors = *GEditor->GetSelectedActors();
 	for( int32 CurSelectedActorIndex = 0; CurSelectedActorIndex < SelectedActors.Num(); ++CurSelectedActorIndex )
 	{
-		AActor* CurSelectedActor = CastChecked<AActor>( SelectedActors.GetSelectedObject( CurSelectedActorIndex ) );
-		UStaticMeshComponent* StaticMeshComponent = NULL;
-		AStaticMeshActor* StaticMeshActor = Cast< AStaticMeshActor >( CurSelectedActor );
-		if( StaticMeshActor != NULL )
+		AActor* SelectedActor = Cast< AActor >( SelectedActors.GetSelectedObject( CurSelectedActorIndex ) );
+		if(SelectedActor != NULL)
 		{
-			StaticMeshComponent = StaticMeshActor->StaticMeshComponent;
-		}
-
-		if( StaticMeshComponent != NULL && StaticMeshComponent->StaticMesh != NULL && StaticMeshComponent->StaticMesh->GetNumLODs() > PaintingMeshLODIndex )
-		{
-			// count the base mesh color data
-			FStaticMeshLODResources& LODModel = StaticMeshComponent->StaticMesh->RenderData->LODResources[ PaintingMeshLODIndex ];
-			OutTotalBaseVertexColorBytes += LODModel.ColorVertexBuffer.GetNumVertices();
-
-			// count the instance color data
-			if( StaticMeshComponent->LODData.Num() > PaintingMeshLODIndex )
+			TArray<UStaticMeshComponent*> StaticMeshComponents;
+			SelectedActor->GetComponents<UStaticMeshComponent>(StaticMeshComponents);
+			for(const auto& StaticMeshComponent : StaticMeshComponents)
 			{
-				const FStaticMeshComponentLODInfo& InstanceMeshLODInfo = StaticMeshComponent->LODData[ PaintingMeshLODIndex ];
-				if( InstanceMeshLODInfo.OverrideVertexColors )
+				if( StaticMeshComponent != NULL && StaticMeshComponent->StaticMesh != NULL && StaticMeshComponent->StaticMesh->GetNumLODs() > PaintingMeshLODIndex )
 				{
-					OutTotalInstanceVertexColorBytes += InstanceMeshLODInfo.OverrideVertexColors->GetAllocatedSize();
+					// count the base mesh color data
+					FStaticMeshLODResources& LODModel = StaticMeshComponent->StaticMesh->RenderData->LODResources[ PaintingMeshLODIndex ];
+					OutTotalBaseVertexColorBytes += LODModel.ColorVertexBuffer.GetNumVertices();
+
+					// count the instance color data
+					if( StaticMeshComponent->LODData.Num() > PaintingMeshLODIndex )
+					{
+						const FStaticMeshComponentLODInfo& InstanceMeshLODInfo = StaticMeshComponent->LODData[ PaintingMeshLODIndex ];
+						if( InstanceMeshLODInfo.OverrideVertexColors )
+						{
+							OutTotalInstanceVertexColorBytes += InstanceMeshLODInfo.OverrideVertexColors->GetAllocatedSize();
+						}
+					}
+
+					++NumValidMeshes;
 				}
 			}
-
-			++NumValidMeshes;
 		}
 	}
 
@@ -4158,11 +4217,15 @@ void FEdModeMeshPaint::GetBrushRadiiLimits( float& OutMinBrushRadius, float& Out
 /** Returns whether there are colors in the copy buffer */
 bool FEdModeMeshPaint::CanPasteVertexColors() const
 { 
-	for (int32 i=0; i<CopiedColorsByLOD.Num(); i++)
+	for (int32 ComponentIndex = 0; ComponentIndex < CopiedColorsByComponent.Num(); ComponentIndex++)
 	{
-		if (0 < CopiedColorsByLOD[i].ColorsByIndex.Num())
+		const FPerComponentVertexColorData& ComponentData = CopiedColorsByComponent[ComponentIndex];
+		for (int32 LODIndex = 0; LODIndex < ComponentData.PerLODVertexColorData.Num(); LODIndex++)
 		{
-			return true;
+			if (0 < ComponentData.PerLODVertexColorData[LODIndex].ColorsByIndex.Num())
+			{
+				return true;
+			}
 		}
 	}
 	return false; 
@@ -4208,25 +4271,26 @@ void FImportVertexTextureHelper::ImportVertexColors(const FString & Filename, in
 		return;
 	}
 
-	TArray<UObject*> Actors;
+	TArray<UStaticMeshComponent*> Components;
 	for ( FSelectionIterator It( GEditor->GetSelectedActorIterator() ) ; It ; ++It )
 	{
-
-		AStaticMeshActor* StaticActor = Cast<AStaticMeshActor>( *It );
-		if(StaticActor)
+		AActor* Actor = Cast<AActor>( *It );
+		if(Actor)
 		{
-			Actors.AddUnique(*It);
+			TArray<UStaticMeshComponent*> ThisActorsComponents;
+			Actor->GetComponents<UStaticMeshComponent>(ThisActorsComponents);
+			Components.Append(ThisActorsComponents);
 		}
 	}
 
-	if (Actors.Num() <1) 
+	if (Components.Num() < 1) 
 	{
 		EditorErrors.Warning(LOCTEXT("MeshPaint_ImportErrNoActors", "No valid actors selected."));
 		EditorErrors.Notify();
 		return;
 	}
 
-	if (Filename.Len() ==0)
+	if (Filename.IsEmpty())
 	{
 		EditorErrors.Warning(LOCTEXT("MeshPaint_ImportErrNoTga", "No tga file specified."));
 		EditorErrors.Notify();
@@ -4267,19 +4331,15 @@ void FImportVertexTextureHelper::ImportVertexColors(const FString & Filename, in
 	uint8* MipData = SrcMipData.GetTypedData();
 	TArray <UStaticMesh*> ModifiedStaticMeshes;
 
-	for (int32 i =0; i < Actors.Num() ; ++i)
+	for(const auto& StaticMeshComponent : Components)
 	{
-		AStaticMeshActor* StaticActor = Cast<AStaticMeshActor>( Actors[i] );
-
-		UStaticMeshComponent* StaticMeshComponent = NULL;
-		UStaticMesh* StaticMesh = NULL; 
-		if (StaticActor)
+		if(StaticMeshComponent == NULL)
 		{
-			StaticMeshComponent = StaticActor->StaticMeshComponent;
-			StaticMesh = StaticMeshComponent->StaticMesh;
+			continue;
 		}
 
-		if (!StaticMeshComponent || !StaticMesh )
+		UStaticMesh* StaticMesh = StaticMeshComponent->StaticMesh; 
+		if(StaticMesh == NULL)
 		{
 			continue;
 		}
@@ -4288,6 +4348,7 @@ void FImportVertexTextureHelper::ImportVertexColors(const FString & Filename, in
 		{
 			continue;
 		}
+
 		FStaticMeshLODResources& LODModel = StaticMesh->RenderData->LODResources[ ImportLOD ];
 
 		TScopedPointer< FStaticMeshComponentRecreateRenderStateContext > RecreateRenderStateContext;
@@ -4438,149 +4499,155 @@ void FEdModeMeshPaint::UpdateTexturePaintTargetList()
 
 		if ( ActorBeingEdited.IsValid() )
 		{
+			AActor* Actor = ActorBeingEdited.Get();
 			const FMeshSelectedMaterialInfo* MeshData = CurrentlySelectedActorsMaterialInfo.Find(ActorBeingEdited);
 			if ( MeshData != NULL )
 			{
 				// Get the selected material index and selected actor from the cached actor info
-				int32 MaterialIndex = MeshData->SelectedMaterialIndex;
+				const int32 MaterialIndex = MeshData->SelectedMaterialIndex;
 
 				// we only operate on static meshes.
-				AStaticMeshActor* StaticMeshActor = Cast< AStaticMeshActor >( ActorBeingEdited.Get() );
-			
-				if ( StaticMeshActor != NULL && StaticMeshActor->StaticMeshComponent != NULL )
+				TArray<UStaticMeshComponent*> StaticMeshComponents;
+				Actor->GetComponents<UStaticMeshComponent>(StaticMeshComponents);
+
+				for(const auto& StaticMeshComponent : StaticMeshComponents)
 				{
-					// We already know the material we are painting on, take it off the static mesh component
-					UMaterialInterface* Material = StaticMeshActor->StaticMeshComponent->GetMaterial(MaterialIndex);
-				
-					if ( Material != NULL )
+					if ( StaticMeshComponent != NULL )
 					{
-						int32 DefaultIndex = INDEX_NONE;
-						FPaintableTexture PaintableTexture;
-						// Find all the unique textures used in the top material level of the selected actor materials
+						// We already know the material we are painting on, take it off the static mesh component
+						UMaterialInterface* Material = StaticMeshComponent->GetMaterial(MaterialIndex);
+				
+						if ( Material != NULL )
+						{
+							int32 DefaultIndex = INDEX_NONE;
+							FPaintableTexture PaintableTexture;
+							// Find all the unique textures used in the top material level of the selected actor materials
 		
-						const TArray<UMaterialExpression*>& Expressions = Material->GetMaterial()->Expressions;
+							const TArray<UMaterialExpression*>& Expressions = Material->GetMaterial()->Expressions;
 
-						// Only grab the textures from the top level of samples
-						for (auto ItExpressions = Expressions.CreateConstIterator(); ItExpressions; ItExpressions++)
-						{
-							UMaterialExpressionTextureBase* TextureBase = Cast<UMaterialExpressionTextureBase>(*ItExpressions);
-							if (TextureBase != NULL && 
-								TextureBase->Texture != NULL && 
-								!TextureBase->Texture->IsNormalMap())
+							// Only grab the textures from the top level of samples
+							for (auto ItExpressions = Expressions.CreateConstIterator(); ItExpressions; ItExpressions++)
 							{
-								// Default UV channel to index 0. 
-								PaintableTexture = FPaintableTexture(TextureBase->Texture, 0);
-										
-								// Texture Samples can have UV's specified, check the first node for whether it has a custom UV channel set. 
-								// We only check the first as the Mesh paint mode does not support painting with UV's modified in the shader.
-								UMaterialExpressionTextureSample* TextureSample = Cast<UMaterialExpressionTextureSample>(*ItExpressions);
-								if (TextureSample != NULL)
+								UMaterialExpressionTextureBase* TextureBase = Cast<UMaterialExpressionTextureBase>(*ItExpressions);
+								if (TextureBase != NULL && 
+									TextureBase->Texture != NULL && 
+									!TextureBase->Texture->IsNormalMap())
 								{
-									UMaterialExpressionTextureCoordinate* TextureCoords = Cast<UMaterialExpressionTextureCoordinate>(TextureSample->Coordinates.Expression);
-									if (TextureCoords != NULL)
+									// Default UV channel to index 0. 
+									PaintableTexture = FPaintableTexture(TextureBase->Texture, 0);
+										
+									// Texture Samples can have UV's specified, check the first node for whether it has a custom UV channel set. 
+									// We only check the first as the Mesh paint mode does not support painting with UV's modified in the shader.
+									UMaterialExpressionTextureSample* TextureSample = Cast<UMaterialExpressionTextureSample>(*ItExpressions);
+									if (TextureSample != NULL)
 									{
-										// Store the uv channel, this is set when the texture is selected. 
-										PaintableTexture.UVChannelIndex = TextureCoords->CoordinateIndex;
+										UMaterialExpressionTextureCoordinate* TextureCoords = Cast<UMaterialExpressionTextureCoordinate>(TextureSample->Coordinates.Expression);
+										if (TextureCoords != NULL)
+										{
+											// Store the uv channel, this is set when the texture is selected. 
+											PaintableTexture.UVChannelIndex = TextureCoords->CoordinateIndex;
+										}
+
+										// Handle texture parameter expressions
+										UMaterialExpressionTextureSampleParameter* TextureSampleParameter = Cast<UMaterialExpressionTextureSampleParameter>(TextureSample);
+										if (TextureSampleParameter != NULL)
+										{
+											// Grab the overridden texture if it exists.  
+											Material->GetTextureParameterValue(TextureSampleParameter->ParameterName, PaintableTexture.Texture);
+										}
 									}
 
-									// Handle texture parameter expressions
-									UMaterialExpressionTextureSampleParameter* TextureSampleParameter = Cast<UMaterialExpressionTextureSampleParameter>(TextureSample);
-									if (TextureSampleParameter != NULL)
+									// note that the same texture will be added again if its UV channel differs. 
+									int32 TextureIndex = TexturesInSelection.AddUnique(PaintableTexture);
+
+									// cache the first default index, if there is no previous info this will be used as the selected texture
+									if (DefaultIndex == INDEX_NONE && TextureBase->IsDefaultMeshpaintTexture)
 									{
-										// Grab the overridden texture if it exists.  
-										Material->GetTextureParameterValue(TextureSampleParameter->ParameterName, PaintableTexture.Texture);
+										DefaultIndex = TextureIndex;
 									}
 								}
-
-								// note that the same texture will be added again if its UV channel differs. 
-								int32 TextureIndex = TexturesInSelection.AddUnique(PaintableTexture);
-
-								// cache the first default index, if there is no previous info this will be used as the selected texture
-								if (DefaultIndex == INDEX_NONE && TextureBase->IsDefaultMeshpaintTexture)
-								{
-									DefaultIndex = TextureIndex;
-								}
 							}
-						}
 
-						// Generate the list of target paint textures that will be displaying in the UI
-						for( int32 TexIndex = 0; TexIndex < TexturesInSelection.Num(); TexIndex++ )
-						{
-							UTexture2D* Texture2D = Cast<UTexture2D>( TexturesInSelection[ TexIndex ].Texture );
-							int32 UVChannelIndex = TexturesInSelection[ TexIndex ].UVChannelIndex;
-							// If this is not a UTexture2D we check to see if it is a rendertarget texture
-							if( Texture2D == NULL )
+							// Generate the list of target paint textures that will be displaying in the UI
+							for( int32 TexIndex = 0; TexIndex < TexturesInSelection.Num(); TexIndex++ )
 							{
-								UTextureRenderTarget2D* TextureRenderTarget2D = Cast<UTextureRenderTarget2D>( TexturesInSelection[ TexIndex ].Texture );
-								if( TextureRenderTarget2D )
+								UTexture2D* Texture2D = Cast<UTexture2D>( TexturesInSelection[ TexIndex ].Texture );
+								int32 UVChannelIndex = TexturesInSelection[ TexIndex ].UVChannelIndex;
+								// If this is not a UTexture2D we check to see if it is a rendertarget texture
+								if( Texture2D == NULL )
 								{
-									// Since this is a rendertarget, we lookup the original texture that we overrode during the paint operation
-									Texture2D = GetOriginalTextureFromRenderTarget( TextureRenderTarget2D );
-
-									// Since we looked up a texture via a rendertarget, it is possible that this texture already exists in our list.  If so 
-									//  we will not add it and continue processing other elements.
-									if( Texture2D != NULL && TexturesInSelection.Contains( FPaintableTexture(Texture2D, UVChannelIndex) ) )
+									UTextureRenderTarget2D* TextureRenderTarget2D = Cast<UTextureRenderTarget2D>( TexturesInSelection[ TexIndex ].Texture );
+									if( TextureRenderTarget2D )
 									{
-										continue;
+										// Since this is a rendertarget, we lookup the original texture that we overrode during the paint operation
+										Texture2D = GetOriginalTextureFromRenderTarget( TextureRenderTarget2D );
+
+										// Since we looked up a texture via a rendertarget, it is possible that this texture already exists in our list.  If so 
+										//  we will not add it and continue processing other elements.
+										if( Texture2D != NULL && TexturesInSelection.Contains( FPaintableTexture(Texture2D, UVChannelIndex) ) )
+										{
+											continue;
+										}
 									}
 								}
-							}
 
-							if( Texture2D != NULL )
-							{
-								// @todo MeshPaint: We rely on filtering out normal maps by name here.  Obviously a user can name a diffuse with _N_ in the name so
-								//   this is not a good option.  We attempted to find all the normal maps from the material above with GetAllNormalParameterNames(),
-								//   but that always seems to return an empty list.  This needs to be revisited.
-	
-								// Some normalmaps in the content will fail checks we do in the if statement below.  So we also check to make sure 
-								//   the name does not end with "_N", and that the following substrings do not appear in the name "_N_" "_N0".
-								FString Texture2DName;
-								Texture2D->GetName(Texture2DName);
-								Texture2DName = Texture2DName.ToUpper();
-
-								// Make sure the texture is not a normalmap, we don't support painting on those at the moment.
-								if( Texture2D->IsNormalMap() == true 
-									|| Texture2D->LODGroup == TEXTUREGROUP_WorldNormalMap
-									|| Texture2D->LODGroup == TEXTUREGROUP_CharacterNormalMap
-									|| Texture2D->LODGroup == TEXTUREGROUP_WeaponNormalMap
-									|| Texture2D->LODGroup == TEXTUREGROUP_VehicleNormalMap
-									|| Texture2D->LODGroup == TEXTUREGROUP_WorldNormalMap
-									|| Texture2DName.Contains( TEXT("_N0" ))
-									|| Texture2DName.Contains( TEXT("_N_" ))
-									|| Texture2DName.Contains( TEXT("_NORMAL" ))
-									|| (Texture2DName.Right(2)).Contains( TEXT("_N" )) )
+								if( Texture2D != NULL )
 								{
-									continue;
-								}
+									// @todo MeshPaint: We rely on filtering out normal maps by name here.  Obviously a user can name a diffuse with _N_ in the name so
+									//   this is not a good option.  We attempted to find all the normal maps from the material above with GetAllNormalParameterNames(),
+									//   but that always seems to return an empty list.  This needs to be revisited.
+	
+									// Some normalmaps in the content will fail checks we do in the if statement below.  So we also check to make sure 
+									//   the name does not end with "_N", and that the following substrings do not appear in the name "_N_" "_N0".
+									FString Texture2DName;
+									Texture2D->GetName(Texture2DName);
+									Texture2DName = Texture2DName.ToUpper();
+
+									// Make sure the texture is not a normalmap, we don't support painting on those at the moment.
+									if( Texture2D->IsNormalMap() == true 
+										|| Texture2D->LODGroup == TEXTUREGROUP_WorldNormalMap
+										|| Texture2D->LODGroup == TEXTUREGROUP_CharacterNormalMap
+										|| Texture2D->LODGroup == TEXTUREGROUP_WeaponNormalMap
+										|| Texture2D->LODGroup == TEXTUREGROUP_VehicleNormalMap
+										|| Texture2D->LODGroup == TEXTUREGROUP_WorldNormalMap
+										|| Texture2DName.Contains( TEXT("_N0" ))
+										|| Texture2DName.Contains( TEXT("_N_" ))
+										|| Texture2DName.Contains( TEXT("_NORMAL" ))
+										|| (Texture2DName.Right(2)).Contains( TEXT("_N" )) )
+									{
+										continue;
+									}
 			
-								// Add the texture to our list
-								new(TexturePaintTargetList) FTextureTargetListInfo(Texture2D, UVChannelIndex);
+									// Add the texture to our list
+									new(TexturePaintTargetList) FTextureTargetListInfo(Texture2D, UVChannelIndex);
 
-								// We stored off the user's selection before we began the update.  Since we cleared the list we lost
-								//  that selection info. If the same texture appears in our list after update, we will select it again.
-								if( PreviouslySelectedTexture != NULL && Texture2D == PreviouslySelectedTexture )
-								{
-									TexturePaintTargetList[ TexturePaintTargetList.Num() - 1 ].bIsSelected = true;
+									// We stored off the user's selection before we began the update.  Since we cleared the list we lost
+									//  that selection info. If the same texture appears in our list after update, we will select it again.
+									if( PreviouslySelectedTexture != NULL && Texture2D == PreviouslySelectedTexture )
+									{
+										TexturePaintTargetList[ TexturePaintTargetList.Num() - 1 ].bIsSelected = true;
+									}
 								}
 							}
-						}
 
-						//if there are no default textures, revert to the old method of just selecting the first texture.
-						if (DefaultIndex == INDEX_NONE)
-						{
-							DefaultIndex = 0;
-						}
+							//if there are no default textures, revert to the old method of just selecting the first texture.
+							if (DefaultIndex == INDEX_NONE)
+							{
+								DefaultIndex = 0;
+							}
 
-						//We refreshed the list, if nothing else is set we default to the first texture that has IsDefaultMeshPaintTexture set.
-						if(TexturePaintTargetList.Num() > 0 && GetSelectedTexture() == NULL)
-						{
-							if (ensure(TexturePaintTargetList.IsValidIndex(DefaultIndex)))
+							//We refreshed the list, if nothing else is set we default to the first texture that has IsDefaultMeshPaintTexture set.
+							if(TexturePaintTargetList.Num() > 0 && GetSelectedTexture() == NULL)
 							{
-								TexturePaintTargetList[DefaultIndex].bIsSelected = true;
+								if (ensure(TexturePaintTargetList.IsValidIndex(DefaultIndex)))
+								{
+									TexturePaintTargetList[DefaultIndex].bIsSelected = true;
+								}
 							}
 						}
 					}
 				}
+		
 			}
 		}
 		
@@ -5082,7 +5149,7 @@ void FEdModeMeshPaint::CacheActorInfo()
 	{
 		TArray<UMaterialInterface*> UsedMaterials;
 
-		// Currently we only support static mesh actors or dynamic static mesh actors
+		// Currently we only support static mesh components
 		UStaticMeshComponent* StaticMesh = SMComponents[CurSMIndex];
 
 		// Get the materials used by the mesh
diff --git a/Engine/Source/Editor/MeshPaint/Public/MeshPaintEdMode.h b/Engine/Source/Editor/MeshPaint/Public/MeshPaintEdMode.h
--- a/Engine/Source/Editor/MeshPaint/Public/MeshPaintEdMode.h
+++ b/Engine/Source/Editor/MeshPaint/Public/MeshPaintEdMode.h
@@ -294,6 +294,25 @@ private:
 		TMap<FVector, FColor> ColorsByPosition;
 	};
 
+	/** struct used to store the color data copied from mesh component to mesh component */
+	struct FPerComponentVertexColorData
+	{
+		FPerComponentVertexColorData(UStaticMesh* InStaticMesh, int32 InComponentIndex)
+			: OriginalMesh(InStaticMesh)
+			, ComponentIndex(InComponentIndex)
+		{
+		}
+
+		/** We match up components by the mesh they use */
+		TWeakObjectPtr<UStaticMesh> OriginalMesh;
+
+		/** We also match by component index */
+		int32 ComponentIndex;
+
+		/** Vertex colors by LOD */
+		TArray<FPerLODVertexColorData> PerLODVertexColorData;
+	};
+
 public:
 
 	struct PaintTexture2DData
@@ -394,7 +413,7 @@ public:
 	/** Helper function to get the current paint action for use in DoPaint */
 	EMeshPaintAction::Type GetPaintAction(FViewport* InViewport);
 
-	/** Removes vertex colors associated with the object (if it's a StaticMeshActor) */
+	/** Removes vertex colors associated with the object */
 	void RemoveInstanceVertexColors(UObject* Obj) const;
 
 	/** Removes vertex colors associated with the static mesh component */
@@ -777,7 +796,7 @@ private:
 	TMap< UTexture2D*, PaintTexture2DData > PaintTargetData;
 
 	/** Temporary buffers used when copying/pasting colors */
-	TArray< FPerLODVertexColorData > CopiedColorsByLOD;
+	TArray<FPerComponentVertexColorData> CopiedColorsByComponent;
 
 	/**
 	 * Does the work of removing instance vertex colors from a single static mesh component.
diff --git a/Engine/Source/Runtime/Engine/Classes/Components/StaticMeshComponent.h b/Engine/Source/Runtime/Engine/Classes/Components/StaticMeshComponent.h
--- a/Engine/Source/Runtime/Engine/Classes/Components/StaticMeshComponent.h
+++ b/Engine/Source/Runtime/Engine/Classes/Components/StaticMeshComponent.h
@@ -116,6 +116,58 @@ public:
 	TArray<FGuid> IrrelevantLights;
 };
 
+/**
+ * Vertex color information stored during RerunConstructionScripts
+ */
+class FVertexColorInstanceData : public FComponentInstanceDataBase
+{
+public:
+	static const FName VertexColorInstanceDataName;
+
+	/** FComponentInstanceDataBase interface */
+	virtual FName GetDataTypeName() const OVERRIDE
+	{
+		return VertexColorInstanceDataName;
+	}
+
+	/** Add vertex color data for a specified LOD before RerunConstructionScripts is called */
+	void AddVertexColorData( const struct FStaticMeshComponentLODInfo& LODInfo, uint32 LODIndex );
+
+	/** Re-apply vertex color data after RerunConstructionScripts is called */
+	bool ApplyVertexColorData( UStaticMeshComponent* StaticMeshComponent );
+
+	/** Check whether this vertex color data can match the specified component */
+	bool MatchesComponent( const UStaticMeshComponent* StaticMeshComponent ) const;
+
+	/** Vertex data stored per-LOD */
+	struct FVertexColorLODData
+	{
+		/** copy of painted vertex data */
+		TArray<FPaintedVertex> PaintedVertices;
+
+		/** Copy of vertex buffer colors */
+		TArray<FColor> VertexBufferColors;
+
+		/** Index of the LOD that this data came from */
+		uint32 LODIndex;
+
+		/** Check whether this contains valid data */
+		bool IsValid() const 
+		{ 
+			return PaintedVertices.Num() > 0 && VertexBufferColors.Num() > 0; 
+		}
+	};
+
+	/** Mesh being used by component */
+	class UStaticMesh* StaticMesh;
+
+	/** 'Stable' index into serialized components array */
+	int32 SerializedComponentsIndex;
+
+	/** Array of cached vertex colors for each LOD */
+	TArray<FVertexColorLODData> VertexColorLODs;
+};
+
 /** A StaticMeshComponent is a mesh that does not animate. */
 UCLASS(HeaderGroup=Component, ClassGroup=(Rendering, Common), hidecategories=(Object,Activation,"Components|Activation"), ShowCategories=(Mobility), dependson=ULightmassPrimitiveSettingsObject, editinlinenew, meta=(BlueprintSpawnableComponent))
 class ENGINE_API UStaticMeshComponent : public UMeshComponent
@@ -437,6 +489,8 @@ public:
 	/** Returns the wireframe color to use for this component. */
 	FColor GetWireframeColor() const;
 
+	/** Get this components index in its parents serialized components array (used for matching instance data) */
+	int32 GetSerializedComponentIndex() const;
 };
 
 
diff --git a/Engine/Source/Runtime/Engine/Private/ColorVertexBuffer.cpp b/Engine/Source/Runtime/Engine/Private/ColorVertexBuffer.cpp
--- a/Engine/Source/Runtime/Engine/Private/ColorVertexBuffer.cpp
+++ b/Engine/Source/Runtime/Engine/Private/ColorVertexBuffer.cpp
@@ -298,6 +298,15 @@ void FColorVertexBuffer::operator=(const FColorVertexBuffer &Other)
 	VertexData = NULL;
 }
 
+void FColorVertexBuffer::GetVertexColors( TArray<FColor>& OutColors )
+{
+	if( VertexData != NULL && NumVertices > 0 )
+	{
+		OutColors.Init( NumVertices );
+
+		FMemory::Memcpy( OutColors.GetData(), VertexData->GetDataPointer(), NumVertices * VertexData->GetStride() ) ;
+	}
+}
 
 /** Load from raw color array */
 void FColorVertexBuffer::InitFromColorArray( const FColor *InColors, const uint32 Count, const uint32 Stride )
diff --git a/Engine/Source/Runtime/Engine/Private/StaticMeshComponent.cpp b/Engine/Source/Runtime/Engine/Private/StaticMeshComponent.cpp
--- a/Engine/Source/Runtime/Engine/Private/StaticMeshComponent.cpp
+++ b/Engine/Source/Runtime/Engine/Private/StaticMeshComponent.cpp
@@ -9,6 +9,71 @@
 
 #define LOCTEXT_NAMESPACE "StaticMeshComponent"
 
+const FName FVertexColorInstanceData::VertexColorInstanceDataName(TEXT("VertexColorInstanceData"));
+
+void FVertexColorInstanceData::AddVertexColorData( const FStaticMeshComponentLODInfo& LODInfo, uint32 LODIndex )
+{
+	if(VertexColorLODs.Num() <= (int32)LODIndex)
+	{
+		VertexColorLODs.SetNum(LODIndex + 1);
+	}
+	FVertexColorLODData& VertexColorData = VertexColorLODs[LODIndex];
+	VertexColorData.LODIndex = LODIndex;
+	VertexColorData.PaintedVertices = LODInfo.PaintedVertices;
+	LODInfo.OverrideVertexColors->GetVertexColors( VertexColorData.VertexBufferColors );
+}
+
+bool FVertexColorInstanceData::ApplyVertexColorData( UStaticMeshComponent* StaticMeshComponent )
+{
+	bool bAppliedAnyData = false;
+
+	if(StaticMeshComponent != NULL)
+	{
+		StaticMeshComponent->SetLODDataCount(VertexColorLODs.Num(), VertexColorLODs.Num());
+
+		for( int32 LODDataIndex = 0; LODDataIndex < VertexColorLODs.Num(); ++LODDataIndex )
+		{
+			const FVertexColorLODData& VertexColorLODData = VertexColorLODs[LODDataIndex];
+			uint32 LODIndex = VertexColorLODData.LODIndex;
+
+			if( StaticMeshComponent->LODData.IsValidIndex( LODIndex ) )
+			{
+				FStaticMeshComponentLODInfo& LODInfo = StaticMeshComponent->LODData[LODIndex];
+			
+				// Should not already have overriden vertex colors
+				if( LODInfo.OverrideVertexColors == NULL )
+				{
+					LODInfo.PaintedVertices = VertexColorLODData.PaintedVertices;
+
+					LODInfo.OverrideVertexColors = new FColorVertexBuffer;
+					LODInfo.OverrideVertexColors->InitFromColorArray( VertexColorLODData.VertexBufferColors );
+
+					BeginInitResource( LODInfo.OverrideVertexColors );	
+					bAppliedAnyData = true;
+				}
+			}
+		}	
+	}
+
+	return bAppliedAnyData;
+}
+
+bool FVertexColorInstanceData::MatchesComponent( const UStaticMeshComponent* StaticMeshComponent ) const
+{
+	if(StaticMeshComponent != NULL)
+	{
+		if(StaticMeshComponent->StaticMesh != StaticMesh)
+		{
+			return false;
+		}
+
+		return SerializedComponentsIndex == StaticMeshComponent->GetSerializedComponentIndex();
+	}
+
+	// this component was not found
+	return false;
+}
+
 UStaticMeshComponent::UStaticMeshComponent(const class FPostConstructInitializeProperties& PCIP)
 	: Super(PCIP)
 {
@@ -1359,6 +1424,22 @@ void UStaticMeshComponent::GetUsedMaterials(TArray<UMaterialInterface*>& OutMate
 	}
 }
 
+int32 UStaticMeshComponent::GetSerializedComponentIndex() const
+{
+	int32 ComponentIndex = 0;
+	for(const auto& Component : GetOwner()->SerializedComponents)
+	{
+		if(Component == this)
+		{
+			return ComponentIndex;
+		}
+
+		ComponentIndex++;
+	}
+
+	return INDEX_NONE;
+}
+
 // Init type name static
 const FName FLightMapInstanceData::LightMapInstanceDataTypeName(TEXT("LightMapInstanceData"));
 
@@ -1384,6 +1465,32 @@ void UStaticMeshComponent::GetComponentInstanceData(FComponentInstanceDataCache&
 		// Add to cache
 		Cache.AddInstanceData(LightMapData);
 	}
+
+
+	// Cache instance vertex colors
+	TSharedPtr<FVertexColorInstanceData> VertexColorData;
+	for( int32 LODIndex = 0; LODIndex < LODData.Num(); ++LODIndex )
+	{
+		const FStaticMeshComponentLODInfo& LODInfo = LODData[LODIndex];
+
+		if ( LODInfo.OverrideVertexColors && LODInfo.OverrideVertexColors->GetNumVertices() > 0 && LODInfo.PaintedVertices.Num() > 0 )
+		{
+			if( !VertexColorData.IsValid() )
+			{
+				VertexColorData = MakeShareable( new FVertexColorInstanceData );
+				VertexColorData->StaticMesh = StaticMesh;
+				VertexColorData->SerializedComponentsIndex = GetSerializedComponentIndex();
+			}
+
+			VertexColorData->AddVertexColorData( LODInfo, LODIndex );
+		}
+	}
+
+	if( VertexColorData.IsValid() )
+	{
+		Cache.AddInstanceData( VertexColorData );
+	}
+
 }
 
 void UStaticMeshComponent::ApplyComponentInstanceData(const FComponentInstanceDataCache& Cache)
@@ -1417,6 +1524,19 @@ void UStaticMeshComponent::ApplyComponentInstanceData(const FComponentInstanceDa
 			MarkRenderStateDirty();
 		}
 	}
+
+	TArray< TSharedPtr<FComponentInstanceDataBase> > AllCachedVertexColorData;
+	Cache.GetInstanceDataOfType( FVertexColorInstanceData::VertexColorInstanceDataName, AllCachedVertexColorData );
+
+	for(const auto& CachedVertexColorData : AllCachedVertexColorData)
+	{
+		TSharedPtr<FVertexColorInstanceData> ComponentInstanceData = StaticCastSharedPtr<FVertexColorInstanceData>(CachedVertexColorData);
+		if(	ComponentInstanceData->MatchesComponent(this) )
+		{
+			ComponentInstanceData->ApplyVertexColorData(this);
+			MarkRenderStateDirty();	
+		}
+	}
 }
 
 void UStaticMeshComponent::CreateRenderState_Concurrent()
diff --git a/Engine/Source/Runtime/Engine/Public/Components.h b/Engine/Source/Runtime/Engine/Public/Components.h
--- a/Engine/Source/Runtime/Engine/Public/Components.h
+++ b/Engine/Source/Runtime/Engine/Public/Components.h
@@ -104,6 +104,13 @@ public:
 	/** Useful for memory profiling. */
 	ENGINE_API uint32 GetAllocatedSize() const;
 
+	/**
+	 * Gets all vertex colors in the buffer
+	 *
+	 * @param OutColors	The populated list of colors 
+	 */
+	ENGINE_API void GetVertexColors( TArray<FColor>& OutColors );
+
 	/**
 	* Load from a array of colors
 	* @param InColors - must not be 0
