Fixed crash when using mesh paint modes texture editor, caused by committing texture changes then changing map.
#TTP 332246 - CRITICAL: Regression: EDITOR: PAINT: TEXTURE: CRASH: Opening new level after commiting changes to a texture

#branch UE4

#proj Editor.MeshPaint

#Change Stopped the texture getting overridden again by the render target following the texture commit. This issue is caused because the committ button clears the override texture, but leaves the editor in such a state that upon exitting the mode the texture wont actually be cleared again unless the user has painted again. This would be fine but every tick the selected texture is set to override with the render texture. This means that when the mode exits the texture is still set in an override state with the render texture, the render texture is then cleaned up leaving the expression cache with a invalid texture pointer. During the new map load all render commands are flushed which involves rendering that mesh using said render target and this is what causes the crash.

diff --git a/Engine/Source/Editor/MeshPaint/Private/MeshPaintEdMode.cpp b/Engine/Source/Editor/MeshPaint/Private/MeshPaintEdMode.cpp
--- a/Engine/Source/Editor/MeshPaint/Private/MeshPaintEdMode.cpp
+++ b/Engine/Source/Editor/MeshPaint/Private/MeshPaintEdMode.cpp
@@ -2602,7 +2602,7 @@ void FEdModeMeshPaint::SetSpecificTextureOverrideForMesh(UStaticMeshComponent* I
 
 		if(bIsTextureUsed)
 		{
-			if(TextureData)
+			if(TextureData && TextureData->PaintingMaterials.Num() > 0)
 			{
 				// If there is texture data, that means we have an override ready, so set it. 
 				MaterialToCheck->OverrideTexture( Texture, TextureData->PaintRenderTargetTexture );
