Set default erase color in Mesh Paint mode to RGB 0,0,0.
#ttp 334531 - Vertex Paint: Erase color defaults to full white (values = 1.0 RGB), which is the same as the paint color. Erase should be RGB 0
#branch UE4
#proj Editor.MeshPaint
#change Changed the initialization of EraseColor in FMeshPaintSettings.

diff --git a/Engine/Source/Editor/MeshPaint/Public/MeshPaintEdMode.h b/Engine/Source/Editor/MeshPaint/Public/MeshPaintEdMode.h
--- a/Engine/Source/Editor/MeshPaint/Public/MeshPaintEdMode.h
+++ b/Engine/Source/Editor/MeshPaint/Public/MeshPaintEdMode.h
@@ -119,7 +119,7 @@ public:
 		  ResourceType( EMeshPaintResource::VertexColors ),
 		  PaintMode( EMeshPaintMode::PaintColors ),
 		  PaintColor( 1.0f, 1.0f, 1.0f, 1.0f ),
-		  EraseColor( 1.0f, 1.0f, 1.0f, 1.0f ),
+		  EraseColor( 0.0f, 0.0f, 0.0f, 1.0f ),
 		  bWriteRed( true ),
 		  bWriteGreen( true ),
 		  bWriteBlue( true ),
