Mobile Preview
 * More removal of GRHIFeatureLevel and GRHIShaderPlatform

diff --git a/Engine/Source/Editor/GraphEditor/Private/MaterialNodes/SGraphNodeMaterialBase.cpp b/Engine/Source/Editor/GraphEditor/Private/MaterialNodes/SGraphNodeMaterialBase.cpp
--- a/Engine/Source/Editor/GraphEditor/Private/MaterialNodes/SGraphNodeMaterialBase.cpp
+++ b/Engine/Source/Editor/GraphEditor/Private/MaterialNodes/SGraphNodeMaterialBase.cpp
@@ -188,7 +188,7 @@ void FPreviewElement::DrawRenderThread(FRHICommandListImmediate& RHICmdList, con
 		float CurrentTime = bIsRealtime ? (FApp::GetCurrentTime() - GStartTime) : 0.0f;
 		float DeltaTime = bIsRealtime ? FApp::GetDeltaTime() : 0.0f;
 
-		FCanvas Canvas(RenderTarget, NULL, CurrentTime, CurrentTime, DeltaTime);
+		FCanvas Canvas(RenderTarget, NULL, CurrentTime, CurrentTime, DeltaTime, GMaxRHIFeatureLevel);
 		{
 			Canvas.SetAllowedModes( 0 );
 			Canvas.SetRenderTargetRect( RenderTarget->GetViewRect() );
diff --git a/Engine/Source/Editor/MeshPaint/Private/MeshPaintEdMode.cpp b/Engine/Source/Editor/MeshPaint/Private/MeshPaintEdMode.cpp
--- a/Engine/Source/Editor/MeshPaint/Private/MeshPaintEdMode.cpp
+++ b/Engine/Source/Editor/MeshPaint/Private/MeshPaintEdMode.cpp
@@ -44,9 +44,9 @@ class FMeshPaintBatchedElementParameters : public FBatchedElementParameters
 public:
 
 	/** Binds vertex and pixel shaders for this element */
-	virtual void BindShaders(FRHICommandList& RHICmdList, const FMatrix& InTransform, const float InGamma, const FMatrix& ColorWeights, const FTexture* Texture) override
+	virtual void BindShaders(FRHICommandList& RHICmdList, ERHIFeatureLevel::Type InFeatureLevel, const FMatrix& InTransform, const float InGamma, const FMatrix& ColorWeights, const FTexture* Texture) override
 	{
-		MeshPaintRendering::SetMeshPaintShaders(RHICmdList, InTransform, InGamma, ShaderParams );
+		MeshPaintRendering::SetMeshPaintShaders(RHICmdList, InFeatureLevel, InTransform, InGamma, ShaderParams );
 	}
 
 public:
@@ -63,9 +63,9 @@ class FMeshPaintDilateBatchedElementParameters : public FBatchedElementParameter
 public:
 
 	/** Binds vertex and pixel shaders for this element */
-	virtual void BindShaders(FRHICommandList& RHICmdList, const FMatrix& InTransform, const float InGamma, const FMatrix& ColorWeights, const FTexture* Texture) override
+	virtual void BindShaders(FRHICommandList& RHICmdList, ERHIFeatureLevel::Type InFeatureLevel, const FMatrix& InTransform, const float InGamma, const FMatrix& ColorWeights, const FTexture* Texture) override
 	{
-		MeshPaintRendering::SetMeshPaintDilateShaders(RHICmdList, InTransform, InGamma, ShaderParams );
+		MeshPaintRendering::SetMeshPaintDilateShaders(RHICmdList, InFeatureLevel, InTransform, InGamma, ShaderParams );
 	}
 
 public:
@@ -1894,6 +1894,8 @@ void FEdModeMeshPaint::StartPaintingTexture( UStaticMeshComponent* InStaticMeshC
 	check( InStaticMeshComponent != NULL );
 	check( TexturePaintingStaticMeshComponent == NULL );
 	check( PaintingTexture2D == NULL );
+
+	const auto FeatureLevel = GetWorld()->FeatureLevel;
 	
 	UTexture2D* Texture2D = GetSelectedTexture();
 	if( Texture2D == NULL )
@@ -2027,7 +2029,7 @@ void FEdModeMeshPaint::StartPaintingTexture( UStaticMeshComponent* InStaticMeshC
 		if( bIsTextureUsed == true && bStartedPainting == true && TextureData->PaintingMaterials.Contains( MaterialToCheck ) == false)
 		{
 			TextureData->PaintingMaterials.AddUnique( MaterialToCheck ); 
-			MaterialToCheck->OverrideTexture( Texture2D, TextureData->PaintRenderTargetTexture );
+			MaterialToCheck->OverrideTexture(Texture2D, TextureData->PaintRenderTargetTexture, FeatureLevel);
 		}
 
 		MaterialIndex++;
@@ -2062,6 +2064,7 @@ void FEdModeMeshPaint::PaintTexture( const FMeshPaintParameters& InParams,
 	FStaticMeshLODResources& LODModel = TexturePaintingStaticMeshComponent->StaticMesh->RenderData->LODResources[ PaintingMeshLODIndex ];
 	FIndexArrayView Indices = LODModel.IndexBuffer.GetArrayView();
 	const uint32 PaintUVCoordinateIndex = InParams.UVChannel;
+	const auto FeatureLevel = GetWorld()->FeatureLevel;
 
 	// Check to see if the UV set is available on the LOD model, if not then there is no point in continuing.
 	if( PaintUVCoordinateIndex >= LODModel.VertexBuffer.GetNumTexCoords() )
@@ -2076,7 +2079,7 @@ void FEdModeMeshPaint::PaintTexture( const FMeshPaintParameters& InParams,
 	// Copy the current image to the brush rendertarget texture.
 	{
 		check( BrushRenderTargetTexture != NULL );
-		CopyTextureToRenderTargetTexture( TextureData->PaintRenderTargetTexture, BrushRenderTargetTexture );
+		CopyTextureToRenderTargetTexture( TextureData->PaintRenderTargetTexture, BrushRenderTargetTexture, FeatureLevel );
 	}
 
 	const bool bEnableSeamPainting = FMeshPaintSettings::Get().bEnableSeamPainting;
@@ -2089,7 +2092,7 @@ void FEdModeMeshPaint::PaintTexture( const FMeshPaintParameters& InParams,
 	check( BrushRenderTargetResource != NULL );
 	
 	// Create a canvas for the brush render target.
-	FCanvas BrushPaintCanvas( BrushRenderTargetResource, NULL, 0, 0, 0 );
+	FCanvas BrushPaintCanvas(BrushRenderTargetResource, NULL, 0, 0, 0, FeatureLevel);
 
 	// Parameters for brush paint
 	TRefCountPtr< FMeshPaintBatchedElementParameters > MeshPaintBatchedElementParameters( new FMeshPaintBatchedElementParameters() );
@@ -2127,7 +2130,7 @@ void FEdModeMeshPaint::PaintTexture( const FMeshPaintParameters& InParams,
 		check( BrushMaskRenderTargetResource != NULL );
 
 		// Create a canvas for the brush mask rendertarget and clear it to black.
-		BrushMaskCanvas = TSharedPtr<FCanvas>( new FCanvas( BrushMaskRenderTargetResource, NULL, 0, 0, 0 ) );
+		BrushMaskCanvas = TSharedPtr<FCanvas>(new FCanvas(BrushMaskRenderTargetResource, NULL, 0, 0, 0, FeatureLevel));
 		BrushMaskCanvas->Clear( FLinearColor::Black );
 
 		// Parameters for the mask
@@ -2303,7 +2306,7 @@ void FEdModeMeshPaint::PaintTexture( const FMeshPaintParameters& InParams,
 	if( !bEnableSeamPainting )
 	{
 		// Seam painting is not enabled so we just copy our delta paint info to the paint target.
-		CopyTextureToRenderTargetTexture( BrushRenderTargetTexture, TextureData->PaintRenderTargetTexture );
+		CopyTextureToRenderTargetTexture(BrushRenderTargetTexture, TextureData->PaintRenderTargetTexture, FeatureLevel);
 	}
 	else
 	{
@@ -2331,7 +2334,7 @@ void FEdModeMeshPaint::PaintTexture( const FMeshPaintParameters& InParams,
 		// Dilate the paint stroke into the texture seams.
 		{
 			// Create a canvas for the render target.
-			FCanvas Canvas3( RenderTargetResource, NULL, 0, 0, 0 );
+			FCanvas Canvas3(RenderTargetResource, NULL, 0, 0, 0, FeatureLevel);
 
 
 			TRefCountPtr< FMeshPaintDilateBatchedElementParameters > MeshPaintDilateBatchedElementParameters( new FMeshPaintDilateBatchedElementParameters() );
@@ -2491,6 +2494,8 @@ void FEdModeMeshPaint::ClearStaticMeshTextureOverrides(UStaticMeshComponent* InS
 		return;
 	}
 
+	const auto FeatureLevel = GetWorld()->FeatureLevel;
+
 	TArray<UMaterialInterface*> UsedMaterials;
 
 	// Get all the used materials for this StaticMeshComponent
@@ -2503,12 +2508,12 @@ void FEdModeMeshPaint::ClearStaticMeshTextureOverrides(UStaticMeshComponent* InS
 		if( Material != NULL )
 		{
 			TArray<UTexture*> UsedTextures;
-			Material->GetUsedTextures( UsedTextures, EMaterialQualityLevel::Num, false, GRHIFeatureLevel, false );
+			Material->GetUsedTextures(UsedTextures, EMaterialQualityLevel::Num, false, FeatureLevel, false);
 
 			for( int32 UsedIndex = 0; UsedIndex < UsedTextures.Num(); UsedIndex++ )
 			{
 				//Reset the texture to it's default.
-				Material->OverrideTexture( UsedTextures[ UsedIndex ], NULL );
+				Material->OverrideTexture(UsedTextures[UsedIndex], NULL, FeatureLevel);
 			}		
 		}
 	}
@@ -2524,7 +2529,7 @@ void FEdModeMeshPaint::ClearAllTextureOverrides()
 		for( int32 MaterialIndex = 0; MaterialIndex < TextureData->PaintingMaterials.Num(); MaterialIndex++)
 		{
 			UMaterialInterface* PaintingMaterialInterface = TextureData->PaintingMaterials[MaterialIndex];
-			PaintingMaterialInterface->OverrideTexture( TextureData->PaintingTexture2D, NULL );
+			PaintingMaterialInterface->OverrideTexture( TextureData->PaintingTexture2D, NULL, GetWorld()->FeatureLevel );
 		}
 
 		TextureData->PaintingMaterials.Empty();
@@ -2539,6 +2544,8 @@ void FEdModeMeshPaint::SetAllTextureOverrides(UStaticMeshComponent* InStaticMesh
 		return;
 	}
 
+	const auto FeatureLevel = GetWorld()->FeatureLevel;
+
 	TArray<UMaterialInterface*> UsedMaterials;
 
 	// Get all the used materials for this StaticMeshComponent
@@ -2553,7 +2560,7 @@ void FEdModeMeshPaint::SetAllTextureOverrides(UStaticMeshComponent* InStaticMesh
 		if( Material != NULL )
 		{
 			TArray<UTexture*> UsedTextures;
-			Material->GetUsedTextures( UsedTextures, EMaterialQualityLevel::Num, false, GRHIFeatureLevel, false );
+			Material->GetUsedTextures(UsedTextures, EMaterialQualityLevel::Num, false, FeatureLevel, false);
 
 			for( int32 UsedIndex = 0; UsedIndex < UsedTextures.Num(); UsedIndex++ )
 			{
@@ -2562,7 +2569,7 @@ void FEdModeMeshPaint::SetAllTextureOverrides(UStaticMeshComponent* InStaticMesh
 
 				if(TextureData)
 				{
-					Material->OverrideTexture( UsedTextures[ UsedIndex ], TextureData->PaintRenderTargetTexture );
+					Material->OverrideTexture(UsedTextures[UsedIndex], TextureData->PaintRenderTargetTexture, FeatureLevel);
 				}
 			}		
 		}
@@ -2575,6 +2582,8 @@ void FEdModeMeshPaint::SetSpecificTextureOverrideForMesh(UStaticMeshComponent* I
 {
 	PaintTexture2DData* TextureData = GetPaintTargetData( (UTexture2D*)Texture );
 
+	const auto FeatureLevel = GetWorld()->FeatureLevel;
+
 	// Check all the materials on the mesh to see if the user texture is there
 	int32 MaterialIndex = 0;
 	UMaterialInterface* MaterialToCheck = InStaticMeshComponent->GetMaterial( MaterialIndex );
@@ -2587,14 +2596,14 @@ void FEdModeMeshPaint::SetSpecificTextureOverrideForMesh(UStaticMeshComponent* I
 			if(TextureData && TextureData->PaintingMaterials.Num() > 0)
 			{
 				// If there is texture data, that means we have an override ready, so set it. 
-				MaterialToCheck->OverrideTexture( Texture, TextureData->PaintRenderTargetTexture );
+				MaterialToCheck->OverrideTexture(Texture, TextureData->PaintRenderTargetTexture, FeatureLevel);
 			}
 			else
 			{
 				// If there is no data, then remove the override so we can at least see the texture without the changes to the other texture.
 					// This is important because overrides are shared between material instances with the same parent. We want to disable a override in place,
 					// making the action more comprehensive to the user.
-				MaterialToCheck->OverrideTexture( Texture, NULL );
+				MaterialToCheck->OverrideTexture(Texture, NULL, FeatureLevel);
 			}
 		}
 
@@ -2832,7 +2841,7 @@ void FEdModeMeshPaint::SaveSettingsForActor( AActor* InActor )
 				}
 
 				TArray<UTexture*> UsedTextures;
-				UsedMaterials[ MatIndex ]->GetUsedTextures( UsedTextures, EMaterialQualityLevel::Num, false, GRHIFeatureLevel, false );
+				UsedMaterials[MatIndex]->GetUsedTextures(UsedTextures, EMaterialQualityLevel::Num, false, GetWorld()->FeatureLevel, false);
 
 				for( int32 TexIndex = 0; TexIndex < UsedTextures.Num(); TexIndex++ )
 				{
@@ -2882,7 +2891,7 @@ void FEdModeMeshPaint::UpdateSettingsForStaticMeshComponent( UStaticMeshComponen
 			}
 
 			TArray<UTexture*> UsedTextures;
-			UsedMaterials[MatIndex]->GetUsedTextures(UsedTextures, EMaterialQualityLevel::Num, false, GRHIFeatureLevel, false);
+			UsedMaterials[MatIndex]->GetUsedTextures(UsedTextures, EMaterialQualityLevel::Num, false, GetWorld()->FeatureLevel, false);
 
 			for( int32 TexIndex = 0; TexIndex < UsedTextures.Num(); TexIndex++ )
 			{
@@ -3147,7 +3156,7 @@ void FEdModeMeshPaint::SetupInitialRenderTargetData( UTexture2D* InTextureSource
 			check( TempSourceArtTexture != NULL );
 
 			// Copy the texture to the render target using the GPU
-			CopyTextureToRenderTargetTexture( TempSourceArtTexture, InRenderTarget );
+			CopyTextureToRenderTargetTexture(TempSourceArtTexture, InRenderTarget, GetWorld()->FeatureLevel);
 
 			// NOTE: TempSourceArtTexture is no longer needed (will be GC'd)
 		}
@@ -3157,7 +3166,7 @@ void FEdModeMeshPaint::SetupInitialRenderTargetData( UTexture2D* InTextureSource
 		// Just copy (render) the texture in GPU memory to our render target.  Hopefully it's not
 		// compressed already!
 		check( InTextureSource->IsFullyStreamedIn() );
-		CopyTextureToRenderTargetTexture( InTextureSource, InRenderTarget );
+		CopyTextureToRenderTargetTexture(InTextureSource, InRenderTarget, GetWorld()->FeatureLevel);
 	}
 
 }
@@ -3214,7 +3223,7 @@ UTexture2D* FEdModeMeshPaint::CreateTempUncompressedTexture( UTexture2D* SourceT
 
 
 /** Static: Copies a texture to a render target texture */
-void FEdModeMeshPaint::CopyTextureToRenderTargetTexture( UTexture* SourceTexture, UTextureRenderTarget2D* RenderTargetTexture )
+void FEdModeMeshPaint::CopyTextureToRenderTargetTexture( UTexture* SourceTexture, UTextureRenderTarget2D* RenderTargetTexture, ERHIFeatureLevel::Type FeatureLevel )
 {	
 	check( SourceTexture != NULL );
 	check( RenderTargetTexture != NULL );
@@ -3228,7 +3237,7 @@ void FEdModeMeshPaint::CopyTextureToRenderTargetTexture( UTexture* SourceTexture
 
 	{
 		// Create a canvas for the render target and clear it to black
-		FCanvas Canvas( RenderTargetResource, NULL, 0, 0, 0 );
+		FCanvas Canvas(RenderTargetResource, NULL, 0, 0, 0, FeatureLevel);
 
 		const uint32 Width = RenderTargetTexture->GetSurfaceWidth();
 		const uint32 Height = RenderTargetTexture->GetSurfaceHeight();
@@ -3407,7 +3416,7 @@ bool FEdModeMeshPaint::GenerateSeamMask(UStaticMeshComponent* StaticMeshComponen
 
 	{
 		// Create a canvas for the render target and clear it to white
-		FCanvas Canvas( RenderTargetResource, NULL, 0, 0, 0 );
+		FCanvas Canvas(RenderTargetResource, NULL, 0, 0, 0, GetWorld()->FeatureLevel);
 		Canvas.Clear( FLinearColor::White);
 		
 		TArray<FCanvasUVTri> TriList;
diff --git a/Engine/Source/Editor/MeshPaint/Public/MeshPaintEdMode.h b/Engine/Source/Editor/MeshPaint/Public/MeshPaintEdMode.h
--- a/Engine/Source/Editor/MeshPaint/Public/MeshPaintEdMode.h
+++ b/Engine/Source/Editor/MeshPaint/Public/MeshPaintEdMode.h
@@ -654,7 +654,7 @@ private:
 	void SetupInitialRenderTargetData( UTexture2D* InTextureSource, UTextureRenderTarget2D* InRenderTarget );
 
 	/** Static: Copies a texture to a render target texture */
-	static void CopyTextureToRenderTargetTexture( UTexture* SourceTexture, UTextureRenderTarget2D* RenderTargetTexture );
+	static void CopyTextureToRenderTargetTexture(UTexture* SourceTexture, UTextureRenderTarget2D* RenderTargetTexture, ERHIFeatureLevel::Type FeatureLevel);
 
 	/** Will generate a mask texture, used for texture dilation, and store it in the passed in rendertarget */
 	bool GenerateSeamMask(UStaticMeshComponent* StaticMeshComponent, int32 UVSet, UTextureRenderTarget2D* RenderTargetTexture);
diff --git a/Engine/Source/Editor/UnrealEd/Private/CubemapUnwrapUtils.cpp b/Engine/Source/Editor/UnrealEd/Private/CubemapUnwrapUtils.cpp
--- a/Engine/Source/Editor/UnrealEd/Private/CubemapUnwrapUtils.cpp
+++ b/Engine/Source/Editor/UnrealEd/Private/CubemapUnwrapUtils.cpp
@@ -45,7 +45,7 @@ namespace CubemapHelpers
 		RenderTargetLongLat->TargetGamma = 0;
 		FRenderTarget* RenderTarget = RenderTargetLongLat->GameThread_GetRenderTargetResource();
 
-		FCanvas* Canvas = new FCanvas(RenderTarget, NULL, 0, 0, 0);
+		FCanvas* Canvas = new FCanvas(RenderTarget, NULL, 0, 0, 0, GMaxRHIFeatureLevel);
 		Canvas->SetRenderTarget_GameThread(RenderTarget);
 
 		// Clear the render target to black
@@ -124,27 +124,27 @@ void FCubemapTexturePropertiesPS<bHDROutput>::SetParameters( FRHICommandList& RH
 	SetShaderValue(RHICmdList, GetPixelShader(), Gamma, GammaValue);
 }
 
-void FMipLevelBatchedElementParameters::BindShaders(FRHICommandList& RHICmdList, const FMatrix& InTransform, const float InGamma, const FMatrix& ColorWeights, const FTexture* Texture)
+void FMipLevelBatchedElementParameters::BindShaders(FRHICommandList& RHICmdList, ERHIFeatureLevel::Type InFeatureLevel, const FMatrix& InTransform, const float InGamma, const FMatrix& ColorWeights, const FTexture* Texture)
 {
 	if(bHDROutput)
 	{
-		BindShaders<FCubemapTexturePropertiesPS<true> >(RHICmdList, InTransform, InGamma, ColorWeights, Texture);
+		BindShaders<FCubemapTexturePropertiesPS<true> >(RHICmdList, InFeatureLevel, InTransform, InGamma, ColorWeights, Texture);
 	}
 	else
 	{
-		BindShaders<FCubemapTexturePropertiesPS<false> >(RHICmdList, InTransform, InGamma, ColorWeights, Texture);
+		BindShaders<FCubemapTexturePropertiesPS<false> >(RHICmdList, InFeatureLevel, InTransform, InGamma, ColorWeights, Texture);
 	}
 }
 
 template<typename TPixelShader>
-void FMipLevelBatchedElementParameters::BindShaders(FRHICommandList& RHICmdList, const FMatrix& InTransform, const float InGamma, const FMatrix& ColorWeights, const FTexture* Texture)
+void FMipLevelBatchedElementParameters::BindShaders(FRHICommandList& RHICmdList, ERHIFeatureLevel::Type InFeatureLevel, const FMatrix& InTransform, const float InGamma, const FMatrix& ColorWeights, const FTexture* Texture)
 {
 	TShaderMapRef<FCubemapTexturePropertiesVS> VertexShader(GetGlobalShaderMap());
 	TShaderMapRef<TPixelShader> PixelShader(GetGlobalShaderMap());
 	
 
 	static FGlobalBoundShaderState BoundShaderState;
-	SetGlobalBoundShaderState(RHICmdList, BoundShaderState, GSimpleElementVertexDeclaration.VertexDeclarationRHI, *VertexShader, *PixelShader);
+	SetGlobalBoundShaderState(RHICmdList, InFeatureLevel, BoundShaderState, GSimpleElementVertexDeclaration.VertexDeclarationRHI, *VertexShader, *PixelShader);
 
 	VertexShader->SetParameters(RHICmdList, InTransform);
 
@@ -161,14 +161,14 @@ void FIESLightProfilePS::SetParameters( FRHICommandList& RHICmdList, const FText
 	SetShaderValue(RHICmdList, ShaderRHI, BrightnessInLumens, InBrightnessInLumens);
 }
 
-void FIESLightProfileBatchedElementParameters::BindShaders( FRHICommandList& RHICmdList, const FMatrix& InTransform, const float InGamma, const FMatrix& ColorWeights, const FTexture* Texture )
+void FIESLightProfileBatchedElementParameters::BindShaders( FRHICommandList& RHICmdList, ERHIFeatureLevel::Type InFeatureLevel, const FMatrix& InTransform, const float InGamma, const FMatrix& ColorWeights, const FTexture* Texture )
 {
 	TShaderMapRef<FSimpleElementVS> VertexShader(GetGlobalShaderMap());
 	TShaderMapRef<FIESLightProfilePS> PixelShader(GetGlobalShaderMap());
 
 	
 	static FGlobalBoundShaderState BoundShaderState;
-	SetGlobalBoundShaderState(RHICmdList, BoundShaderState, GSimpleElementVertexDeclaration.VertexDeclarationRHI, *VertexShader, *PixelShader);
+	SetGlobalBoundShaderState(RHICmdList, InFeatureLevel, BoundShaderState, GSimpleElementVertexDeclaration.VertexDeclarationRHI, *VertexShader, *PixelShader);
 	RHICmdList.SetBlendState(TStaticBlendState<>::GetRHI());
 
 	VertexShader->SetParameters(RHICmdList, InTransform);
diff --git a/Engine/Source/Editor/UnrealEd/Private/EditorExporters.cpp b/Engine/Source/Editor/UnrealEd/Private/EditorExporters.cpp
--- a/Engine/Source/Editor/UnrealEd/Private/EditorExporters.cpp
+++ b/Engine/Source/Editor/UnrealEd/Private/EditorExporters.cpp
@@ -1318,7 +1318,7 @@ bool GenerateExportMaterialPropertyData(
 			RenderTarget->ClearColor = FLinearColor(0.0f, 0.0f, 0.0f, 0.0f);
 			RenderTarget->InitCustomFormat(InOutSizeX, InOutSizeY, Format, bForceLinear);
 
-			Canvas = new FCanvas(RenderTarget->GetRenderTargetResource(), NULL, 0, 0, 0);
+			Canvas = new FCanvas(RenderTarget->GetRenderTargetResource(), NULL, 0, 0, 0, GMaxRHIFeatureLevel);
 			check(Canvas);
 		}
 
@@ -2610,7 +2610,7 @@ namespace MaterialExportUtils
 
 		{
 			// Create a canvas for the render target and clear it to black
-			FCanvas Canvas(RTResource, NULL, FApp::GetCurrentTime() - GStartTime, FApp::GetDeltaTime(), FApp::GetCurrentTime() - GStartTime);
+			FCanvas Canvas(RTResource, NULL, FApp::GetCurrentTime() - GStartTime, FApp::GetDeltaTime(), FApp::GetCurrentTime() - GStartTime, GMaxRHIFeatureLevel);
 			Canvas.Clear(FLinearColor::Black);
 			FCanvasTileItem TileItem(FVector2D(0.0f, 0.0f), MaterialProxy, FVector2D(InRenderTarget->SizeX, InRenderTarget->SizeY));
 			TileItem.bFreezeTime = true;
@@ -2773,7 +2773,7 @@ namespace MaterialExportUtils
 
 				ViewFamily.Views.Add(new FSceneView(ViewInitOptions));
 					
-				FCanvas Canvas(RenderTargetResource, NULL, FApp::GetCurrentTime() - GStartTime, FApp::GetDeltaTime(), FApp::GetCurrentTime() - GStartTime);
+				FCanvas Canvas(RenderTargetResource, NULL, FApp::GetCurrentTime() - GStartTime, FApp::GetDeltaTime(), FApp::GetCurrentTime() - GStartTime, Scene->GetFeatureLevel());
 				Canvas.Clear(FLinearColor::Black);
 				GetRendererModule().BeginRenderingViewFamily(&Canvas, &ViewFamily);
 
@@ -2857,7 +2857,7 @@ namespace MaterialExportUtils
 				NewView->CurrentBufferVisualizationMode = FName("WorldNormal");
 				ViewFamily.Views.Add(NewView);
 										
-				FCanvas Canvas(RenderTargetResource, NULL, FApp::GetCurrentTime() - GStartTime, FApp::GetDeltaTime(), FApp::GetCurrentTime() - GStartTime);
+				FCanvas Canvas(RenderTargetResource, NULL, FApp::GetCurrentTime() - GStartTime, FApp::GetDeltaTime(), FApp::GetCurrentTime() - GStartTime, Scene->GetFeatureLevel());
 				Canvas.Clear(FLinearColor::Black);
 				GetRendererModule().BeginRenderingViewFamily(&Canvas, &ViewFamily);
 
diff --git a/Engine/Source/Editor/UnrealEd/Private/Lightmass/LightmassRender.cpp b/Engine/Source/Editor/UnrealEd/Private/Lightmass/LightmassRender.cpp
--- a/Engine/Source/Editor/UnrealEd/Private/Lightmass/LightmassRender.cpp
+++ b/Engine/Source/Editor/UnrealEd/Private/Lightmass/LightmassRender.cpp
@@ -1056,7 +1056,7 @@ bool FLightmassMaterialRenderer::CreateRenderTarget(EPixelFormat InFormat, int32
 		RenderTarget->ClearColor = FLinearColor(0.0f, 0.0f, 0.0f, 0.0f);
 		RenderTarget->InitCustomFormat(InSizeX, InSizeY, InFormat, false);
 
-		Canvas = new FCanvas(RenderTarget->GameThread_GetRenderTargetResource(), NULL, 0, 0, 0);
+		Canvas = new FCanvas(RenderTarget->GameThread_GetRenderTargetResource(), NULL, 0, 0, 0, GMaxRHIFeatureLevel);
 		check(Canvas);
 	}
 
diff --git a/Engine/Source/Editor/UnrealEd/Private/MeshPaintRendering.cpp b/Engine/Source/Editor/UnrealEd/Private/MeshPaintRendering.cpp
--- a/Engine/Source/Editor/UnrealEd/Private/MeshPaintRendering.cpp
+++ b/Engine/Source/Editor/UnrealEd/Private/MeshPaintRendering.cpp
@@ -348,7 +348,7 @@ namespace MeshPaintRendering
 
 
 	/** Binds the mesh paint vertex and pixel shaders to the graphics device */
-	void SetMeshPaintShaders(FRHICommandList& RHICmdList, const FMatrix& InTransform,
+	void SetMeshPaintShaders(FRHICommandList& RHICmdList, ERHIFeatureLevel::Type InFeatureLevel, const FMatrix& InTransform,
 										   const float InGamma,
 										   const FMeshPaintShaderParameters& InShaderParams )
 	{
@@ -357,7 +357,7 @@ namespace MeshPaintRendering
 
 		
 		static FGlobalBoundShaderState BoundShaderState;
-		SetGlobalBoundShaderState(RHICmdList, BoundShaderState, GMeshPaintVertexDeclaration.VertexDeclarationRHI, *VertexShader, *PixelShader);
+		SetGlobalBoundShaderState(RHICmdList, InFeatureLevel, BoundShaderState, GMeshPaintVertexDeclaration.VertexDeclarationRHI, *VertexShader, *PixelShader);
 
 		// Set vertex shader parameters
 		VertexShader->SetParameters(RHICmdList, InTransform );
@@ -369,7 +369,7 @@ namespace MeshPaintRendering
 	}
 
 	/** Binds the mesh paint vertex and pixel shaders to the graphics device */
-	void SetMeshPaintDilateShaders(FRHICommandList& RHICmdList, const FMatrix& InTransform,
+	void SetMeshPaintDilateShaders(FRHICommandList& RHICmdList, ERHIFeatureLevel::Type InFeatureLevel, const FMatrix& InTransform,
 												 const float InGamma,
 												 const FMeshPaintDilateShaderParameters& InShaderParams )
 	{
@@ -378,7 +378,7 @@ namespace MeshPaintRendering
 
 		
 		static FGlobalBoundShaderState BoundShaderState;
-		SetGlobalBoundShaderState(RHICmdList, BoundShaderState, GMeshPaintDilateVertexDeclaration.VertexDeclarationRHI, *VertexShader, *PixelShader);
+		SetGlobalBoundShaderState(RHICmdList, InFeatureLevel, BoundShaderState, GMeshPaintDilateVertexDeclaration.VertexDeclarationRHI, *VertexShader, *PixelShader);
 
 		// Set vertex shader parameters
 		VertexShader->SetParameters(RHICmdList, InTransform );
diff --git a/Engine/Source/Editor/UnrealEd/Private/NormalMapPreview.cpp b/Engine/Source/Editor/UnrealEd/Private/NormalMapPreview.cpp
--- a/Engine/Source/Editor/UnrealEd/Private/NormalMapPreview.cpp
+++ b/Engine/Source/Editor/UnrealEd/Private/NormalMapPreview.cpp
@@ -76,6 +76,7 @@ IMPLEMENT_SHADER_TYPE(,FSimpleElementNormalMapPS,TEXT("SimpleElementNormalMapPix
 /** Binds vertex and pixel shaders for this element */
 void FNormalMapBatchedElementParameters::BindShaders(
 	FRHICommandList& RHICmdList,
+	ERHIFeatureLevel::Type InFeatureLevel,
 	const FMatrix& InTransform,
 	const float InGamma,
 	const FMatrix& ColorWeights,
@@ -86,7 +87,7 @@ void FNormalMapBatchedElementParameters::BindShaders(
 
 	
 	static FGlobalBoundShaderState BoundShaderState;
-	SetGlobalBoundShaderState(RHICmdList, BoundShaderState, GSimpleElementVertexDeclaration.VertexDeclarationRHI, *VertexShader, *PixelShader);
+	SetGlobalBoundShaderState(RHICmdList, InFeatureLevel, BoundShaderState, GSimpleElementVertexDeclaration.VertexDeclarationRHI, *VertexShader, *PixelShader);
 
 	VertexShader->SetParameters(RHICmdList, InTransform);
 	PixelShader->SetParameters(RHICmdList, Texture);
diff --git a/Engine/Source/Editor/UnrealEd/Private/ObjectTools.cpp b/Engine/Source/Editor/UnrealEd/Private/ObjectTools.cpp
--- a/Engine/Source/Editor/UnrealEd/Private/ObjectTools.cpp
+++ b/Engine/Source/Editor/UnrealEd/Private/ObjectTools.cpp
@@ -3406,7 +3406,7 @@ namespace ThumbnailTools
 		check( RenderTargetResource != NULL );
 
 		// Create a canvas for the render target and clear it to black
-		FCanvas Canvas( RenderTargetResource, NULL, FApp::GetCurrentTime() - GStartTime, FApp::GetDeltaTime(), FApp::GetCurrentTime() - GStartTime );
+		FCanvas Canvas( RenderTargetResource, NULL, FApp::GetCurrentTime() - GStartTime, FApp::GetDeltaTime(), FApp::GetCurrentTime() - GStartTime, GMaxRHIFeatureLevel );
 		Canvas.Clear( FLinearColor::Black );
 
 
diff --git a/Engine/Source/Editor/UnrealEd/Private/PlayLevel.cpp b/Engine/Source/Editor/UnrealEd/Private/PlayLevel.cpp
--- a/Engine/Source/Editor/UnrealEd/Private/PlayLevel.cpp
+++ b/Engine/Source/Editor/UnrealEd/Private/PlayLevel.cpp
@@ -1038,7 +1038,7 @@ void UEditorEngine::PlayStandaloneLocalPc(FString MapNameOverride, FIntPoint* Wi
 	// apply additional settings
 	if (bPlayUsingMobilePreview)
 	{
-		if (IsOpenGLPlatform(GRHIShaderPlatform))
+		if (IsOpenGLPlatform(GShaderPlatformForFeatureLevel[GMaxRHIFeatureLevel]))
 		{
 			AdditionalParameters += TEXT(" -opengl");
 		}
@@ -2260,7 +2260,7 @@ UGameInstance* UEditorEngine::CreatePIEGameInstance(int32 PIEInstance, bool bInS
 	FFormatNamedArguments Args;
 	Args.Add( TEXT("GameName"), FText::FromString( FString( FApp::GetGameName() ) ) );
 	Args.Add( TEXT("PlatformBits"), FText::FromString( PlatformBitsString ) );
-	Args.Add( TEXT("RHIName"), FText::FromName( LegacyShaderPlatformToShaderFormat( GRHIShaderPlatform ) ) );
+	Args.Add( TEXT("RHIName"), FText::FromName( LegacyShaderPlatformToShaderFormat( GShaderPlatformForFeatureLevel[GMaxRHIFeatureLevel] ) ) );
 
 	const ULevelEditorPlaySettings* PlayInSettings = GetDefault<ULevelEditorPlaySettings>();
 	if (PlayInSettings->PlayNetMode == PIE_Client)
diff --git a/Engine/Source/Editor/UnrealEd/Private/Texture2DPreview.cpp b/Engine/Source/Editor/UnrealEd/Private/Texture2DPreview.cpp
--- a/Engine/Source/Editor/UnrealEd/Private/Texture2DPreview.cpp
+++ b/Engine/Source/Editor/UnrealEd/Private/Texture2DPreview.cpp
@@ -78,6 +78,7 @@ IMPLEMENT_SHADER_TYPE(,FSimpleElementTexture2DPreviewPS,TEXT("SimpleElementTextu
 /** Binds vertex and pixel shaders for this element */
 void FBatchedElementTexture2DPreviewParameters::BindShaders(
 	FRHICommandList& RHICmdList,
+	ERHIFeatureLevel::Type InFeatureLevel,
 	const FMatrix& InTransform,
 	const float InGamma,
 	const FMatrix& ColorWeights,
@@ -88,7 +89,7 @@ void FBatchedElementTexture2DPreviewParameters::BindShaders(
 
 	
 	static FGlobalBoundShaderState BoundShaderState;
-	SetGlobalBoundShaderState(RHICmdList, BoundShaderState, GSimpleElementVertexDeclaration.VertexDeclarationRHI, *VertexShader, *PixelShader);
+	SetGlobalBoundShaderState(RHICmdList, InFeatureLevel, BoundShaderState, GSimpleElementVertexDeclaration.VertexDeclarationRHI, *VertexShader, *PixelShader);
 
 	VertexShader->SetParameters(RHICmdList, InTransform);
 	PixelShader->SetParameters(RHICmdList, Texture, ColorWeights, InGamma, MipLevel, bIsNormalMap);
diff --git a/Engine/Source/Editor/UnrealEd/Public/CubemapUnwrapUtils.h b/Engine/Source/Editor/UnrealEd/Public/CubemapUnwrapUtils.h
--- a/Engine/Source/Editor/UnrealEd/Public/CubemapUnwrapUtils.h
+++ b/Engine/Source/Editor/UnrealEd/Public/CubemapUnwrapUtils.h
@@ -120,10 +120,10 @@ public:
 	}
 
 	/** Binds vertex and pixel shaders for this element */
-	virtual void BindShaders(FRHICommandList& RHICmdList, const FMatrix& InTransform, const float InGamma, const FMatrix& ColorWeights, const FTexture* Texture) override;
+	virtual void BindShaders(FRHICommandList& RHICmdList, ERHIFeatureLevel::Type InFeatureLevel, const FMatrix& InTransform, const float InGamma, const FMatrix& ColorWeights, const FTexture* Texture) override;
 
 private:
-	template<typename TPixelShader> void BindShaders(FRHICommandList& RHICmdList, const FMatrix& InTransform, const float InGamma, const FMatrix& ColorWeights, const FTexture* Texture);
+	template<typename TPixelShader> void BindShaders(FRHICommandList& RHICmdList, ERHIFeatureLevel::Type InFeatureLevel, const FMatrix& InTransform, const float InGamma, const FMatrix& ColorWeights, const FTexture* Texture);
 
 	bool bHDROutput;
 	/** Parameters that need to be passed to the shader */
@@ -180,7 +180,7 @@ public:
 	}
 
 	/** Binds vertex and pixel shaders for this element */
-	virtual void BindShaders(FRHICommandList& RHICmdList, const FMatrix& InTransform, const float InGamma, const FMatrix& ColorWeights, const FTexture* Texture) override;
+	virtual void BindShaders(FRHICommandList& RHICmdList, ERHIFeatureLevel::Type InFeatureLevel, const FMatrix& InTransform, const float InGamma, const FMatrix& ColorWeights, const FTexture* Texture) override;
 
 private:
 	float BrightnessInLumens;
diff --git a/Engine/Source/Editor/UnrealEd/Public/MeshPaintRendering.h b/Engine/Source/Editor/UnrealEd/Public/MeshPaintRendering.h
--- a/Engine/Source/Editor/UnrealEd/Public/MeshPaintRendering.h
+++ b/Engine/Source/Editor/UnrealEd/Public/MeshPaintRendering.h
@@ -57,14 +57,18 @@ namespace MeshPaintRendering
 
 
 	/** Binds the mesh paint vertex and pixel shaders to the graphics device */
-	void UNREALED_API SetMeshPaintShaders(FRHICommandList& RHICmdList, const FMatrix& InTransform,
-										   const float InGamma,
-										   const FMeshPaintShaderParameters& InShaderParams );
+	void UNREALED_API SetMeshPaintShaders(FRHICommandList& RHICmdList, 
+											ERHIFeatureLevel::Type InFeatureLevel, 
+											const FMatrix& InTransform,
+											const float InGamma,
+											const FMeshPaintShaderParameters& InShaderParams );
 
 	/** Binds the mesh paint dilation vertex and pixel shaders to the graphics device */
-	void UNREALED_API SetMeshPaintDilateShaders(FRHICommandList& RHICmdList, const FMatrix& InTransform,
-												const float InGamma,
-												const FMeshPaintDilateShaderParameters& InShaderParams );
+	void UNREALED_API SetMeshPaintDilateShaders(FRHICommandList& RHICmdList, 
+													ERHIFeatureLevel::Type InFeatureLevel, 
+													const FMatrix& InTransform,
+													const float InGamma,
+													const FMeshPaintDilateShaderParameters& InShaderParams );
 
 }
 
diff --git a/Engine/Source/Editor/UnrealEd/Public/NormalMapPreview.h b/Engine/Source/Editor/UnrealEd/Public/NormalMapPreview.h
--- a/Engine/Source/Editor/UnrealEd/Public/NormalMapPreview.h
+++ b/Engine/Source/Editor/UnrealEd/Public/NormalMapPreview.h
@@ -16,6 +16,7 @@ class UNREALED_API FNormalMapBatchedElementParameters : public FBatchedElementPa
 	/** Binds vertex and pixel shaders for this element */
 	virtual void BindShaders(
 		FRHICommandList& RHICmdList,
+		ERHIFeatureLevel::Type InFeatureLevel,
 		const FMatrix& InTransform,
 		const float InGamma,
 		const FMatrix& ColorWeights,
diff --git a/Engine/Source/Editor/UnrealEd/Public/Texture2DPreview.h b/Engine/Source/Editor/UnrealEd/Public/Texture2DPreview.h
--- a/Engine/Source/Editor/UnrealEd/Public/Texture2DPreview.h
+++ b/Engine/Source/Editor/UnrealEd/Public/Texture2DPreview.h
@@ -22,7 +22,7 @@ public:
 	}
 
 	/** Binds vertex and pixel shaders for this element */
-	virtual void BindShaders(FRHICommandList& RHICmdList, const FMatrix& InTransform, const float InGamma, const FMatrix& ColorWeights, const FTexture* Texture) override;
+	virtual void BindShaders(FRHICommandList& RHICmdList, ERHIFeatureLevel::Type InFeatureLevel, const FMatrix& InTransform, const float InGamma, const FMatrix& ColorWeights, const FTexture* Texture) override;
 
 private:
 
diff --git a/Engine/Source/Runtime/Engine/Classes/Materials/Material.h b/Engine/Source/Runtime/Engine/Classes/Materials/Material.h
--- a/Engine/Source/Runtime/Engine/Classes/Materials/Material.h
+++ b/Engine/Source/Runtime/Engine/Classes/Materials/Material.h
@@ -759,7 +759,7 @@ public:
 	ENGINE_API virtual FMaterialRenderProxy* GetRenderProxy(bool Selected, bool bHovered=false) const override;
 	ENGINE_API virtual UPhysicalMaterial* GetPhysicalMaterial() const override;
 	ENGINE_API virtual void GetUsedTextures(TArray<UTexture*>& OutTextures, EMaterialQualityLevel::Type QualityLevel, bool bAllQualityLevels, ERHIFeatureLevel::Type FeatureLevel, bool bAllFeatureLevels) const override;
-	ENGINE_API virtual void OverrideTexture( const UTexture* InTextureToOverride, UTexture* OverrideTexture ) override;
+	ENGINE_API virtual void OverrideTexture( const UTexture* InTextureToOverride, UTexture* OverrideTexture, ERHIFeatureLevel::Type InFeatureLevel ) override;
 	ENGINE_API virtual bool CheckMaterialUsage(const EMaterialUsage Usage, const bool bSkipPrim = false) override;
 	ENGINE_API virtual bool CheckMaterialUsage_Concurrent(const EMaterialUsage Usage, const bool bSkipPrim = false) const override;
 	ENGINE_API virtual FMaterialResource* AllocateResource();
@@ -984,7 +984,7 @@ public:
 	/**
 	 * Go through every material, flush the specified types and re-initialize the material's shader maps.
 	 */
-	ENGINE_API static void UpdateMaterialShaders(TArray<FShaderType*>& ShaderTypesToFlush, TArray<const FVertexFactoryType*>& VFTypesToFlush, EShaderPlatform ShaderPlatform=GRHIShaderPlatform);
+	ENGINE_API static void UpdateMaterialShaders(TArray<FShaderType*>& ShaderTypesToFlush, TArray<const FVertexFactoryType*>& VFTypesToFlush, EShaderPlatform ShaderPlatform);
 
 	/** 
 	 * Backs up all material shaders to memory through serialization, organized by FMaterialShaderMap. 
diff --git a/Engine/Source/Runtime/Engine/Classes/Materials/MaterialInstance.h b/Engine/Source/Runtime/Engine/Classes/Materials/MaterialInstance.h
--- a/Engine/Source/Runtime/Engine/Classes/Materials/MaterialInstance.h
+++ b/Engine/Source/Runtime/Engine/Classes/Materials/MaterialInstance.h
@@ -219,7 +219,7 @@ public:
 	virtual bool GetTextureParameterValue(FName ParameterName,class UTexture*& OutValue) const override;
 	virtual bool GetVectorParameterValue(FName ParameterName,FLinearColor& OutValue) const override;
 	virtual void GetUsedTextures(TArray<UTexture*>& OutTextures, EMaterialQualityLevel::Type QualityLevel, bool bAllQualityLevels, ERHIFeatureLevel::Type FeatureLevel, bool bAllFeatureLevels) const override;
-	virtual void OverrideTexture( const UTexture* InTextureToOverride, UTexture* OverrideTexture ) override;
+	virtual void OverrideTexture(const UTexture* InTextureToOverride, UTexture* OverrideTexture, ERHIFeatureLevel::Type InFeatureLevel) override;
 	virtual bool CheckMaterialUsage(const EMaterialUsage Usage, const bool bSkipPrim = false) override;
 	virtual bool CheckMaterialUsage_Concurrent(const EMaterialUsage Usage, const bool bSkipPrim = false) const override;
 	virtual bool GetStaticSwitchParameterValue(FName ParameterName,bool &OutValue,FGuid &OutExpressionGuid) override;
diff --git a/Engine/Source/Runtime/Engine/Classes/Materials/MaterialInterface.h b/Engine/Source/Runtime/Engine/Classes/Materials/MaterialInterface.h
--- a/Engine/Source/Runtime/Engine/Classes/Materials/MaterialInterface.h
+++ b/Engine/Source/Runtime/Engine/Classes/Materials/MaterialInterface.h
@@ -272,7 +272,7 @@ public:
 	 * @param InTextureToOverride The texture to override
 	 * @param OverrideTexture The new texture to use
 	 */
-	virtual void OverrideTexture( const UTexture* InTextureToOverride, UTexture* OverrideTexture ) PURE_VIRTUAL(UMaterialInterface::OverrideTexture,return;);
+	virtual void OverrideTexture(const UTexture* InTextureToOverride, UTexture* OverrideTexture, ERHIFeatureLevel::Type InFeatureLevel) PURE_VIRTUAL(UMaterialInterface::OverrideTexture, return;);
 
 	/**
 	 * Checks if the material can be used with the given usage flag.  
diff --git a/Engine/Source/Runtime/Engine/Private/BatchedElements.cpp b/Engine/Source/Runtime/Engine/Private/BatchedElements.cpp
--- a/Engine/Source/Runtime/Engine/Private/BatchedElements.cpp
+++ b/Engine/Source/Runtime/Engine/Private/BatchedElements.cpp
@@ -363,6 +363,7 @@ float GBatchedElementSmoothWidth = 12;
  */
 void FBatchedElements::PrepareShaders(
 	FRHICommandList& RHICmdList,
+	ERHIFeatureLevel::Type FeatureLevel,
 	ESimpleElementBlendMode BlendMode,
 	const FMatrix& Transform,
 	bool bSwitchVerticalAxis,
@@ -442,7 +443,7 @@ void FBatchedElements::PrepareShaders(
 	if( BatchedElementParameters != NULL )
 	{
 		// Use the vertex/pixel shader that we were given
-		BatchedElementParameters->BindShaders(RHICmdList, Transform, Gamma, ColorWeights, Texture );
+		BatchedElementParameters->BindShaders(RHICmdList, FeatureLevel, Transform, Gamma, ColorWeights, Texture);
 	}
 	else
 	{
@@ -451,7 +452,7 @@ void FBatchedElements::PrepareShaders(
 	    if (bHitTesting)
 	    {
 		    TShaderMapRef<FSimpleElementHitProxyPS> HitTestingPixelShader(GetGlobalShaderMap());
-			SetGlobalBoundShaderState(RHICmdList, HitTestingBoundShaderState, GSimpleElementVertexDeclaration.VertexDeclarationRHI,
+			SetGlobalBoundShaderState(RHICmdList, FeatureLevel, HitTestingBoundShaderState, GSimpleElementVertexDeclaration.VertexDeclarationRHI,
 				*VertexShader, *HitTestingPixelShader);
 
 			HitTestingPixelShader->SetParameters(RHICmdList, Texture);
@@ -467,7 +468,7 @@ void FBatchedElements::PrepareShaders(
 				RHICmdList.SetBlendState(TStaticBlendState<>::GetRHI());
 
 			    TShaderMapRef<FSimpleElementMaskedGammaPS> MaskedPixelShader(GetGlobalShaderMap());
-				SetGlobalBoundShaderState(RHICmdList, MaskedBoundShaderState, GSimpleElementVertexDeclaration.VertexDeclarationRHI,
+				SetGlobalBoundShaderState(RHICmdList, FeatureLevel, MaskedBoundShaderState, GSimpleElementVertexDeclaration.VertexDeclarationRHI,
 					*VertexShader, *MaskedPixelShader);
 
 				MaskedPixelShader->SetEditorCompositingParameters(RHICmdList, View, DepthTexture );
@@ -495,7 +496,7 @@ void FBatchedElements::PrepareShaders(
 			    }
 			    
 			    TShaderMapRef<FSimpleElementDistanceFieldGammaPS> DistanceFieldPixelShader(GetGlobalShaderMap());			
-				SetGlobalBoundShaderState(RHICmdList, DistanceFieldBoundShaderState, GSimpleElementVertexDeclaration.VertexDeclarationRHI,
+				SetGlobalBoundShaderState(RHICmdList, FeatureLevel, DistanceFieldBoundShaderState, GSimpleElementVertexDeclaration.VertexDeclarationRHI,
 					*VertexShader, *DistanceFieldPixelShader );			
 
 			    // @todo - expose these as options for batch rendering
@@ -525,7 +526,7 @@ void FBatchedElements::PrepareShaders(
 			else if(BlendMode >= SE_BLEND_RGBA_MASK_START && BlendMode <= SE_BLEND_RGBA_MASK_END)
 			{
 				TShaderMapRef<FSimpleElementColorChannelMaskPS> ColorChannelMaskPixelShader(GetGlobalShaderMap());
-				SetGlobalBoundShaderState(RHICmdList, ColorChannelMaskShaderState, GSimpleElementVertexDeclaration.VertexDeclarationRHI,
+				SetGlobalBoundShaderState(RHICmdList, FeatureLevel, ColorChannelMaskShaderState, GSimpleElementVertexDeclaration.VertexDeclarationRHI,
 					*VertexShader, *ColorChannelMaskPixelShader );
 			
 				ColorChannelMaskPixelShader->SetParameters(RHICmdList, Texture, ColorWeights, GammaToUse );
@@ -537,7 +538,7 @@ void FBatchedElements::PrepareShaders(
 			    if (FMath::Abs(Gamma - 1.0f) < KINDA_SMALL_NUMBER)
 			    {
 				    TShaderMapRef<FSimpleElementPS> RegularPixelShader(GetGlobalShaderMap());
-					SetGlobalBoundShaderState(RHICmdList, SimpleBoundShaderState, GSimpleElementVertexDeclaration.VertexDeclarationRHI,
+					SetGlobalBoundShaderState(RHICmdList, FeatureLevel, SimpleBoundShaderState, GSimpleElementVertexDeclaration.VertexDeclarationRHI,
 						*VertexShader, *RegularPixelShader );
 
 					RegularPixelShader->SetEditorCompositingParameters(RHICmdList, View, DepthTexture);
@@ -546,7 +547,7 @@ void FBatchedElements::PrepareShaders(
 			    else
 			    {
 				    TShaderMapRef<FSimpleElementGammaPS> RegularPixelShader(GetGlobalShaderMap());
-					SetGlobalBoundShaderState(RHICmdList, RegularBoundShaderState, GSimpleElementVertexDeclaration.VertexDeclarationRHI,
+					SetGlobalBoundShaderState(RHICmdList, FeatureLevel, RegularBoundShaderState, GSimpleElementVertexDeclaration.VertexDeclarationRHI,
 						*VertexShader, *RegularPixelShader );
 
 				    RegularPixelShader->SetParameters(RHICmdList, Texture,Gamma,BlendMode);
@@ -608,7 +609,7 @@ void FBatchedElements::DrawPointElements(FRHICommandList& RHICmdList, const FMat
 }
 
 
-bool FBatchedElements::Draw(FRHICommandList& RHICmdList, bool bNeedToSwitchVerticalAxis, const FMatrix& Transform, uint32 ViewportSizeX, uint32 ViewportSizeY, bool bHitTesting, float Gamma, const FSceneView* View, FTexture2DRHIRef DepthTexture, EBlendModeFilter::Type Filter) const
+bool FBatchedElements::Draw(FRHICommandList& RHICmdList, ERHIFeatureLevel::Type FeatureLevel, bool bNeedToSwitchVerticalAxis, const FMatrix& Transform, uint32 ViewportSizeX, uint32 ViewportSizeY, bool bHitTesting, float Gamma, const FSceneView* View, FTexture2DRHIRef DepthTexture, EBlendModeFilter::Type Filter) const
 {
 	if( HasPrimsToDraw() )
 	{
@@ -626,7 +627,7 @@ bool FBatchedElements::Draw(FRHICommandList& RHICmdList, bool bNeedToSwitchVerti
 			FBatchedElementParameters* BatchedElementParameters = NULL;
 
 			// Set the appropriate pixel shader parameters & shader state for the non-textured elements.
-			PrepareShaders(RHICmdList, SE_BLEND_Opaque, Transform, bNeedToSwitchVerticalAxis, BatchedElementParameters, GWhiteTexture, bHitTesting, Gamma, NULL, View, DepthTexture);
+			PrepareShaders(RHICmdList, FeatureLevel, SE_BLEND_Opaque, Transform, bNeedToSwitchVerticalAxis, BatchedElementParameters, GWhiteTexture, bHitTesting, Gamma, NULL, View, DepthTexture);
 
 			// Draw the line elements.
 			if( LineVertices.Num() > 0 )
@@ -842,7 +843,7 @@ bool FBatchedElements::Draw(FRHICommandList& RHICmdList, bool bNeedToSwitchVerti
 						//New batch, draw previous and clear
 						const int32 VertexCount = SpriteList.Num();
 						const int32 PrimCount = VertexCount / 3;
-						PrepareShaders(RHICmdList, CurrentBlendMode, Transform, bNeedToSwitchVerticalAxis, BatchedElementParameters, CurrentTexture, bHitTesting, Gamma, NULL, View, DepthTexture);
+						PrepareShaders(RHICmdList, FeatureLevel, CurrentBlendMode, Transform, bNeedToSwitchVerticalAxis, BatchedElementParameters, CurrentTexture, bHitTesting, Gamma, NULL, View, DepthTexture);
 						DrawPrimitiveUP(RHICmdList, PT_TriangleList, PrimCount, SpriteList.GetTypedData(), sizeof(FSimpleElementVertex));
 
 						SpriteList.Empty(6);
@@ -882,7 +883,7 @@ bool FBatchedElements::Draw(FRHICommandList& RHICmdList, bool bNeedToSwitchVerti
 					//Draw last batch
 					const int32 VertexCount = SpriteList.Num();
 					const int32 PrimCount = VertexCount / 3;
-					PrepareShaders(RHICmdList, CurrentBlendMode, Transform, bNeedToSwitchVerticalAxis, BatchedElementParameters, CurrentTexture, bHitTesting, Gamma, NULL, View, DepthTexture);
+					PrepareShaders(RHICmdList, FeatureLevel, CurrentBlendMode, Transform, bNeedToSwitchVerticalAxis, BatchedElementParameters, CurrentTexture, bHitTesting, Gamma, NULL, View, DepthTexture);
 					DrawPrimitiveUP(RHICmdList, PT_TriangleList, PrimCount, SpriteList.GetTypedData(), sizeof(FSimpleElementVertex));
 				}
 			}
@@ -900,7 +901,7 @@ bool FBatchedElements::Draw(FRHICommandList& RHICmdList, bool bNeedToSwitchVerti
 				if (Filter & MeshFilter)
 				{
 					// Set the appropriate pixel shader for the mesh.
-					PrepareShaders(RHICmdList, MeshElement.BlendMode, Transform, bNeedToSwitchVerticalAxis, MeshElement.BatchedElementParameters, MeshElement.Texture, bHitTesting, Gamma, &MeshElement.GlowInfo);
+					PrepareShaders(RHICmdList, FeatureLevel, MeshElement.BlendMode, Transform, bNeedToSwitchVerticalAxis, MeshElement.BatchedElementParameters, MeshElement.Texture, bHitTesting, Gamma, &MeshElement.GlowInfo);
 
 					// Draw the mesh.
 					DrawIndexedPrimitiveUP(
@@ -930,7 +931,7 @@ bool FBatchedElements::Draw(FRHICommandList& RHICmdList, bool bNeedToSwitchVerti
 					FBatchedElementParameters* BatchedElementParameters = NULL;
 
 					// Set the appropriate pixel shader for the mesh.
-					PrepareShaders(RHICmdList, MeshElement.BlendMode, Transform, bNeedToSwitchVerticalAxis, BatchedElementParameters, MeshElement.Texture, bHitTesting, Gamma);
+					PrepareShaders(RHICmdList, FeatureLevel, MeshElement.BlendMode, Transform, bNeedToSwitchVerticalAxis, BatchedElementParameters, MeshElement.Texture, bHitTesting, Gamma);
 
 					// Draw the mesh.
 					DrawPrimitiveUP(
diff --git a/Engine/Source/Runtime/Engine/Private/CanvasRenderTarget2D.cpp b/Engine/Source/Runtime/Engine/Private/CanvasRenderTarget2D.cpp
--- a/Engine/Source/Runtime/Engine/Private/CanvasRenderTarget2D.cpp
+++ b/Engine/Source/Runtime/Engine/Private/CanvasRenderTarget2D.cpp
@@ -50,7 +50,7 @@ void UCanvasRenderTarget2D::UpdateResource()
 	);
 
 	// Create the FCanvas which does the actual rendering.
-	FCanvas RenderCanvas(GameThread_GetRenderTargetResource(), nullptr, FApp::GetCurrentTime() - GStartTime, FApp::GetDeltaTime(), FApp::GetCurrentTime() - GStartTime);
+	FCanvas RenderCanvas(GameThread_GetRenderTargetResource(), nullptr, FApp::GetCurrentTime() - GStartTime, FApp::GetDeltaTime(), FApp::GetCurrentTime() - GStartTime, GetWorld()->FeatureLevel);
 	Canvas->Canvas = &RenderCanvas;
 
 	if (!IsPendingKill() && OnCanvasRenderTargetUpdate.IsBound())
diff --git a/Engine/Source/Runtime/Engine/Private/Materials/Material.cpp b/Engine/Source/Runtime/Engine/Private/Materials/Material.cpp
--- a/Engine/Source/Runtime/Engine/Private/Materials/Material.cpp
+++ b/Engine/Source/Runtime/Engine/Private/Materials/Material.cpp
@@ -623,13 +623,15 @@ void UMaterial::GetUsedTextures(TArray<UTexture*>& OutTextures, EMaterialQuality
 	}
 }
 
-void UMaterial::OverrideTexture( const UTexture* InTextureToOverride, UTexture* OverrideTexture )
+void UMaterial::OverrideTexture(const UTexture* InTextureToOverride, UTexture* OverrideTexture, ERHIFeatureLevel::Type InFeatureLevel)
 {
 #if WITH_EDITOR
 	bool bShouldRecacheMaterialExpressions = false;
 	const bool bES2Preview = false;
-	ERHIFeatureLevel::Type FeatureLevelsToUpdate[2] = {GRHIFeatureLevel,ERHIFeatureLevel::ES2};
-	int32 NumFeatureLevelsToUpdate = bES2Preview ? 2 : 1;
+	//ERHIFeatureLevel::Type FeatureLevelsToUpdate[2] = { InFeatureLevel, ERHIFeatureLevel::ES2 };
+	//int32 NumFeatureLevelsToUpdate = bES2Preview ? 2 : 1;
+	ERHIFeatureLevel::Type FeatureLevelsToUpdate[1] = { InFeatureLevel };
+	int32 NumFeatureLevelsToUpdate = 1;
 	
 	for (int32 i = 0; i < NumFeatureLevelsToUpdate; ++i)
 	{
diff --git a/Engine/Source/Runtime/Engine/Private/Materials/MaterialInstance.cpp b/Engine/Source/Runtime/Engine/Private/Materials/MaterialInstance.cpp
--- a/Engine/Source/Runtime/Engine/Private/Materials/MaterialInstance.cpp
+++ b/Engine/Source/Runtime/Engine/Private/Materials/MaterialInstance.cpp
@@ -746,13 +746,15 @@ void UMaterialInstance::GetUsedTextures(TArray<UTexture*>& OutTextures, EMateria
 
 
 
-void UMaterialInstance::OverrideTexture( const UTexture* InTextureToOverride, UTexture* OverrideTexture )
+void UMaterialInstance::OverrideTexture(const UTexture* InTextureToOverride, UTexture* OverrideTexture, ERHIFeatureLevel::Type InFeatureLevel)
 {
 #if WITH_EDITOR
 	bool bShouldRecacheMaterialExpressions = false;
 	const bool bES2Preview = false;
-	ERHIFeatureLevel::Type FeatureLevelsToUpdate[2] = { GRHIFeatureLevel, ERHIFeatureLevel::ES2 };
-	int32 NumFeatureLevelsToUpdate = bES2Preview ? 2 : 1;
+	//ERHIFeatureLevel::Type FeatureLevelsToUpdate[2] = { GRHIFeatureLevel, ERHIFeatureLevel::ES2 };
+	//int32 NumFeatureLevelsToUpdate = bES2Preview ? 2 : 1;
+	ERHIFeatureLevel::Type FeatureLevelsToUpdate[1] = { InFeatureLevel };
+	int32 NumFeatureLevelsToUpdate = 1;
 	
 	for (int32 i = 0; i < NumFeatureLevelsToUpdate; ++i)
 	{
diff --git a/Engine/Source/Runtime/Engine/Private/NavMeshRenderingHelpers.h b/Engine/Source/Runtime/Engine/Private/NavMeshRenderingHelpers.h
--- a/Engine/Source/Runtime/Engine/Private/NavMeshRenderingHelpers.h
+++ b/Engine/Source/Runtime/Engine/Private/NavMeshRenderingHelpers.h
@@ -295,10 +295,12 @@ public:
 
 		const bool bNeedToSwitchVerticalAxis = RHINeedsToSwitchVerticalAxis(GRHIShaderPlatform);
 		const FTexture2DRHIRef DepthTexture;
+		const auto FeatureLevel = View->GetFeatureLevel();
 			
 		FRHICommandListImmediate& RHICmdList = FRHICommandListExecutor::GetImmediateCommandList();
 		BatchedElements.Draw(
 			RHICmdList,
+			FeatureLevel,
 			bNeedToSwitchVerticalAxis,
 			View->ViewProjectionMatrix,
 			View->ViewRect.Width(),
@@ -311,6 +313,7 @@ public:
 
 		ProxyData.BatchedElements.Draw(
 			RHICmdList,
+			FeatureLevel,
 			bNeedToSwitchVerticalAxis,
 			View->ViewProjectionMatrix,
 			View->ViewRect.Width(),
diff --git a/Engine/Source/Runtime/Engine/Private/Particles/ParticleCurveTexture.cpp b/Engine/Source/Runtime/Engine/Private/Particles/ParticleCurveTexture.cpp
--- a/Engine/Source/Runtime/Engine/Private/Particles/ParticleCurveTexture.cpp
+++ b/Engine/Source/Runtime/Engine/Private/Particles/ParticleCurveTexture.cpp
@@ -229,6 +229,7 @@ static void InjectCurves(
 		static FGlobalBoundShaderState BoundShaderState;
 		SetGlobalBoundShaderState(
 			RHICmdList,
+			GMaxRHIFeatureLevel,
 			BoundShaderState,
 			GParticleCurveInjectionVertexDeclaration.VertexDeclarationRHI,
 			*VertexShader,
diff --git a/Engine/Source/Runtime/Engine/Private/Particles/ParticleGpuSimulation.cpp b/Engine/Source/Runtime/Engine/Private/Particles/ParticleGpuSimulation.cpp
--- a/Engine/Source/Runtime/Engine/Private/Particles/ParticleGpuSimulation.cpp
+++ b/Engine/Source/Runtime/Engine/Private/Particles/ParticleGpuSimulation.cpp
@@ -1223,6 +1223,7 @@ struct FSimulationCommandGPU
 template <bool bUseDepthBufferCollision>
 void ExecuteSimulationCommands(
 	FRHICommandList& RHICmdList,
+	ERHIFeatureLevel::Type FeatureLevel,
 	const TArray<FSimulationCommandGPU>& SimulationCommands,
 	const FParticleStateTextures& TextureResources,
 	const FParticleAttributesTexture& AttributeTexture,
@@ -1241,6 +1242,7 @@ void ExecuteSimulationCommands(
 	static FGlobalBoundShaderState BoundShaderState;
 	SetGlobalBoundShaderState(
 		RHICmdList,
+		FeatureLevel,
 		BoundShaderState,
 		GParticleTileVertexDeclaration.VertexDeclarationRHI,
 		*VertexShader,
@@ -1273,7 +1275,7 @@ void ExecuteSimulationCommands(
  * Invokes the clear simulation shader for each particle in each tile.
  * @param Tiles - The list of tiles to clear.
  */
-void ClearTiles(FRHICommandList& RHICmdList, const TArray<uint32>& Tiles)
+void ClearTiles(FRHICommandList& RHICmdList, ERHIFeatureLevel::Type FeatureLevel, const TArray<uint32>& Tiles)
 {
 	SCOPED_DRAW_EVENT(ClearTiles, DEC_PARTICLE);
 
@@ -1297,6 +1299,7 @@ void ClearTiles(FRHICommandList& RHICmdList, const TArray<uint32>& Tiles)
 	static FGlobalBoundShaderState BoundShaderState;
 	SetGlobalBoundShaderState(
 		RHICmdList,
+		FeatureLevel,
 		BoundShaderState,
 		GParticleTileVertexDeclaration.VertexDeclarationRHI,
 		*VertexShader,
@@ -1514,7 +1517,7 @@ TGlobalResource<FParticleInjectionVertexDeclaration> GParticleInjectionVertexDec
  * Injects new particles in to the GPU simulation.
  * @param NewParticles - A list of particles to inject in to the simulation.
  */
-void InjectNewParticles(FRHICommandList& RHICmdList, const TArray<FNewParticle>& NewParticles)
+void InjectNewParticles(FRHICommandList& RHICmdList, ERHIFeatureLevel::Type FeatureLevel, const TArray<FNewParticle>& NewParticles)
 {
 	const int32 MaxParticlesPerDrawCall = GParticleScratchVertexBufferSize / sizeof(FNewParticle);
 	FVertexBufferRHIParamRef ScratchVertexBufferRHI = GParticleScratchVertexBuffer.VertexBufferRHI;
@@ -1542,6 +1545,7 @@ void InjectNewParticles(FRHICommandList& RHICmdList, const TArray<FNewParticle>&
 		static FGlobalBoundShaderState BoundShaderState;
 		SetGlobalBoundShaderState(
 			RHICmdList,
+			FeatureLevel,
 			BoundShaderState,
 			GParticleInjectionVertexDeclaration.VertexDeclarationRHI,
 			*VertexShader,
@@ -1731,6 +1735,7 @@ TGlobalResource<FParticleSimVisualizeVertexDeclaration> GParticleSimVisualizeVer
  */
 static void VisualizeGPUSimulation(
 	FRHICommandList& RHICmdList,
+	ERHIFeatureLevel::Type FeatureLevel,
 	int32 VisualizationMode,
 	FRenderTarget* RenderTarget,
 	const FParticleStateTextures& StateTextures,
@@ -1763,6 +1768,7 @@ static void VisualizeGPUSimulation(
 	static FGlobalBoundShaderState BoundShaderState;
 	SetGlobalBoundShaderState(
 		RHICmdList,
+		FeatureLevel,
 		BoundShaderState,
 		GParticleSimVisualizeVertexDeclaration.VertexDeclarationRHI,
 		*VertexShader,
@@ -4259,6 +4265,7 @@ void FFXSystem::SimulateGPUParticles(
 			/// ?
 			ExecuteSimulationCommands<true>(
 				RHICmdList,
+				FeatureLevel,
 				SimulationCommands,
 				PrevStateTextures,
 				ParticleSimulationResources->SimulationAttributesTexture,
@@ -4273,6 +4280,7 @@ void FFXSystem::SimulateGPUParticles(
 			/// ?
 			ExecuteSimulationCommands<false>(
 				RHICmdList,
+				FeatureLevel,
 				SimulationCommands,
 				PrevStateTextures,
 				ParticleSimulationResources->SimulationAttributesTexture,
@@ -4287,7 +4295,7 @@ void FFXSystem::SimulateGPUParticles(
 	// Clear any newly allocated tiles.
 	if (TilesToClear.Num())
 	{
-		ClearTiles(RHICmdList, TilesToClear);
+		ClearTiles(RHICmdList, FeatureLevel, TilesToClear);
 	}
 
 	// Inject any new particles that have spawned into the simulation.
@@ -4306,7 +4314,7 @@ void FFXSystem::SimulateGPUParticles(
 		RHICmdList.SetViewport(0, 0, 0.0f, GParticleSimulationTextureSizeX, GParticleSimulationTextureSizeY, 1.0f);
 
 		// Inject particles.
-		InjectNewParticles(RHICmdList, NewParticles);
+		InjectNewParticles(RHICmdList, FeatureLevel, NewParticles);
 
 		// Resolve attributes textures. State textures are resolved later.
 		RHICmdList.CopyToResolveTarget(
@@ -4343,15 +4351,16 @@ void FFXSystem::SimulateGPUParticles(
 
 void FFXSystem::VisualizeGPUParticles(FCanvas* Canvas)
 {
-	ENQUEUE_UNIQUE_RENDER_COMMAND_THREEPARAMETER(
+	ENQUEUE_UNIQUE_RENDER_COMMAND_FOURPARAMETER(
 		FVisualizeGPUParticlesCommand,
 		FFXSystem*, FXSystem, this,
 		int32, VisualizationMode, FXConsoleVariables::VisualizeGPUSimulation,
 		FRenderTarget*, RenderTarget, Canvas->GetRenderTarget(),
+		ERHIFeatureLevel::Type, FeatureLevel, GetFeatureLevel(),
 	{
 		FParticleSimulationResources* Resources = FXSystem->GetParticleSimulationResources();
 		FParticleStateTextures& CurrentStateTextures = Resources->GetCurrentStateTextures();
-		VisualizeGPUSimulation(RHICmdList, VisualizationMode, RenderTarget, CurrentStateTextures, GParticleCurveTexture.GetCurveTexture());
+		VisualizeGPUSimulation(RHICmdList, FeatureLevel, VisualizationMode, RenderTarget, CurrentStateTextures, GParticleCurveTexture.GetCurveTexture());
 	});
 }
 
diff --git a/Engine/Source/Runtime/Engine/Private/Particles/ParticleSystemRender.cpp b/Engine/Source/Runtime/Engine/Private/Particles/ParticleSystemRender.cpp
--- a/Engine/Source/Runtime/Engine/Private/Particles/ParticleSystemRender.cpp
+++ b/Engine/Source/Runtime/Engine/Private/Particles/ParticleSystemRender.cpp
@@ -2067,7 +2067,6 @@ void FDynamicMeshEmitterData::PreRenderView(FParticleSystemSceneProxy* Proxy, co
 	}
 
 	const auto FeatureLevel = ViewFamily->Scene->GetFeatureLevel();
-	const auto ShaderPlatform = GShaderPlatformForFeatureLevel[FeatureLevel];
 	const bool bInstanced = FeatureLevel >= ERHIFeatureLevel::SM3;
 
 	int32 ParticleCount = Source.ActiveParticleCount;
@@ -2188,7 +2187,7 @@ void FDynamicMeshEmitterData::PreRenderView(FParticleSystemSceneProxy* Proxy, co
 					if (bIsWireframe)
 					{
 						if( LODModel.WireframeIndexBuffer.IsInitialized()
-							&& !(RHISupportsTessellation(ShaderPlatform) && Mesh.VertexFactory->GetType()->SupportsTessellationShaders())
+							&& !(RHISupportsTessellation(GShaderPlatformForFeatureLevel[FeatureLevel]) && Mesh.VertexFactory->GetType()->SupportsTessellationShaders())
 							)
 						{
 							Mesh.Type = PT_LineList;
diff --git a/Engine/Source/Runtime/Engine/Private/RendererSupport.cpp b/Engine/Source/Runtime/Engine/Private/RendererSupport.cpp
--- a/Engine/Source/Runtime/Engine/Private/RendererSupport.cpp
+++ b/Engine/Source/Runtime/Engine/Private/RendererSupport.cpp
@@ -154,7 +154,7 @@ void RestoreReferencesToRendererModuleClasses(
 
 	// Recompile any missing shaders
 	BeginRecompileGlobalShaders(OutdatedShaderTypes);
-	UMaterial::UpdateMaterialShaders(OutdatedShaderTypes, OutdatedFactoryTypes);
+	UMaterial::UpdateMaterialShaders(OutdatedShaderTypes, OutdatedFactoryTypes, GRHIShaderPlatform);
 
 	// Block on global shader jobs
 	FinishRecompileGlobalShaders();
diff --git a/Engine/Source/Runtime/Engine/Private/SceneManagement.cpp b/Engine/Source/Runtime/Engine/Private/SceneManagement.cpp
--- a/Engine/Source/Runtime/Engine/Private/SceneManagement.cpp
+++ b/Engine/Source/Runtime/Engine/Private/SceneManagement.cpp
@@ -111,6 +111,7 @@ void FSimpleElementCollector::DrawBatchedElements(FRHICommandList& RHICmdList, c
 	// Draw the batched elements.
 	BatchedElements.Draw(
 		RHICmdList,
+		View.GetFeatureLevel(),
 		bNeedToSwitchVerticalAxis,
 		View.ViewProjectionMatrix,
 		View.ViewRect.Width(),
diff --git a/Engine/Source/Runtime/Engine/Private/SceneView.cpp b/Engine/Source/Runtime/Engine/Private/SceneView.cpp
--- a/Engine/Source/Runtime/Engine/Private/SceneView.cpp
+++ b/Engine/Source/Runtime/Engine/Private/SceneView.cpp
@@ -988,8 +988,8 @@ void FSceneView::EndFinalPostprocessSettings()
 
 		if( !Family->EngineShowFlags.PostProcessing || !Family->EngineShowFlags.AntiAliasing || Quality <= 0
 			// Disable antialiasing in GammaLDR mode to avoid jittering.
-			|| (GRHIFeatureLevel == ERHIFeatureLevel::ES2 && MobileHDRCvar->GetValueOnGameThread() == 0)
-			|| (GRHIFeatureLevel == ERHIFeatureLevel::ES2 && (MSAAValue == 2 || MSAAValue == 4)) )
+			|| (FeatureLevel == ERHIFeatureLevel::ES2 && MobileHDRCvar->GetValueOnGameThread() == 0)
+			|| (FeatureLevel == ERHIFeatureLevel::ES2 && (MSAAValue == 2 || MSAAValue == 4)))
 		{
 			FinalPostProcessSettings.AntiAliasingMethod = AAM_None;
 		}
@@ -1145,6 +1145,11 @@ ERHIFeatureLevel::Type FSceneView::GetFeatureLevel() const
 	return Family->GetFeatureLevel();
 }
 
+EShaderPlatform FSceneView::GetShaderPlatform() const
+{
+	return GShaderPlatformForFeatureLevel[GetFeatureLevel()];
+}
+
 FSceneViewFamily::FSceneViewFamily( const ConstructionValues& CVS )
 	:
 	FamilySizeX(0),
diff --git a/Engine/Source/Runtime/Engine/Private/ShaderCompiler/ShaderCompiler.cpp b/Engine/Source/Runtime/Engine/Private/ShaderCompiler/ShaderCompiler.cpp
--- a/Engine/Source/Runtime/Engine/Private/ShaderCompiler/ShaderCompiler.cpp
+++ b/Engine/Source/Runtime/Engine/Private/ShaderCompiler/ShaderCompiler.cpp
@@ -2281,7 +2281,7 @@ bool RecompileShaders(const TCHAR* Cmd, FOutputDevice& Ar)
 
 				// Kick off global shader recompiles
 				BeginRecompileGlobalShaders(OutdatedShaderTypes);
-				UMaterial::UpdateMaterialShaders(OutdatedShaderTypes, OutdatedFactoryTypes);
+				UMaterial::UpdateMaterialShaders(OutdatedShaderTypes, OutdatedFactoryTypes, GRHIShaderPlatform);
 
 				GWarn->StatusUpdate(0, 1, NSLOCTEXT("ShaderCompilingManager", "CompilingGlobalShaderStatus", "Compiling global shaders..."));
 
@@ -2355,7 +2355,7 @@ bool RecompileShaders(const TCHAR* Cmd, FOutputDevice& Ar)
 				TArray<const FVertexFactoryType*> FactoryTypes;
 
 				BeginRecompileGlobalShaders(ShaderTypes);
-				UMaterial::UpdateMaterialShaders(ShaderTypes, FactoryTypes);
+				UMaterial::UpdateMaterialShaders(ShaderTypes, FactoryTypes, GRHIShaderPlatform);
 				FinishRecompileGlobalShaders();
 			}
 		}
diff --git a/Engine/Source/Runtime/Engine/Private/SkeletalMesh.cpp b/Engine/Source/Runtime/Engine/Private/SkeletalMesh.cpp
--- a/Engine/Source/Runtime/Engine/Private/SkeletalMesh.cpp
+++ b/Engine/Source/Runtime/Engine/Private/SkeletalMesh.cpp
@@ -4355,6 +4355,8 @@ void FSkeletalMeshSceneProxy::DrawDynamicElementsSection(FPrimitiveDrawInterface
 	}
 #endif // #if WITH_EDITOR
 
+	const auto FeatureLevel = View->GetFeatureLevel();
+
 	const bool bIsWireframe = View->Family->EngineShowFlags.Wireframe;
 
 	if(Section.bDisabled)
@@ -4381,7 +4383,7 @@ void FSkeletalMeshSceneProxy::DrawDynamicElementsSection(FPrimitiveDrawInterface
 
 	BatchElement.UserIndex = MeshObject->GPUSkinCacheKeys[Section.ChunkIndex];
 
-	const bool bRequiresAdjacencyInformation = RequiresAdjacencyInformation(SectionElementInfo.Material, Mesh.VertexFactory->GetType(), View->GetFeatureLevel());
+	const bool bRequiresAdjacencyInformation = RequiresAdjacencyInformation(SectionElementInfo.Material, Mesh.VertexFactory->GetType(), FeatureLevel);
 	if ( bRequiresAdjacencyInformation )
 	{
 		check( LODModel.AdjacencyMultiSizeIndexContainer.IsIndexBufferValid() );
@@ -4437,7 +4439,7 @@ void FSkeletalMeshSceneProxy::DrawDynamicElementsSection(FPrimitiveDrawInterface
 	BatchElement.NumPrimitives = Section.NumTriangles;
 	if( GIsEditor && MeshObject->ProgressiveDrawingFraction != 1.f )
 	{
-		if (Mesh.MaterialRenderProxy->GetMaterial(GRHIFeatureLevel)->GetBlendMode() == BLEND_Translucent)
+		if (Mesh.MaterialRenderProxy->GetMaterial(FeatureLevel)->GetBlendMode() == BLEND_Translucent)
 		{
 			BatchElement.NumPrimitives = FMath::RoundToInt(((float)Section.NumTriangles)*FMath::Clamp<float>(MeshObject->ProgressiveDrawingFraction,0.f,1.f));
 			if( BatchElement.NumPrimitives == 0 )
diff --git a/Engine/Source/Runtime/Engine/Private/Slate/DebugCanvas.cpp b/Engine/Source/Runtime/Engine/Private/Slate/DebugCanvas.cpp
--- a/Engine/Source/Runtime/Engine/Private/Slate/DebugCanvas.cpp
+++ b/Engine/Source/Runtime/Engine/Private/Slate/DebugCanvas.cpp
@@ -27,7 +27,7 @@ class FCanvasProxy
 {
 public:
 	FCanvasProxy( FRenderTarget* RenderTarget, UWorld* InWorld )
-		: Canvas( RenderTarget, NULL, InWorld )
+		: Canvas(RenderTarget, NULL, InWorld, InWorld->FeatureLevel)
 	{
 		// Do not allow the canvas to be flushed outside of our debug rendering path
 		Canvas.SetAllowedModes( 0 );
diff --git a/Engine/Source/Runtime/Engine/Private/StaticBoundShaderState.cpp b/Engine/Source/Runtime/Engine/Private/StaticBoundShaderState.cpp
--- a/Engine/Source/Runtime/Engine/Private/StaticBoundShaderState.cpp
+++ b/Engine/Source/Runtime/Engine/Private/StaticBoundShaderState.cpp
@@ -100,15 +100,17 @@ void FGlobalBoundShaderStateResource::ReleaseRHI()
 	BoundShaderState.SafeRelease();
 }
 
-static FBoundShaderStateRHIParamRef GetGlobalBoundShaderState_Internal(FGlobalBoundShaderState& GlobalBoundShaderState)
+static FBoundShaderStateRHIParamRef GetGlobalBoundShaderState_Internal(FGlobalBoundShaderState& GlobalBoundShaderState, ERHIFeatureLevel::Type InFeatureLevel)
 {
+	auto WorkArea = GlobalBoundShaderState.Get(InFeatureLevel);
+
 	// Check for unset uniform buffer parameters
 	// Technically you can set uniform buffer parameters after calling RHISetBoundShaderState, but this is the most global place to check for unset parameters
-	GlobalBoundShaderState.Get()->Args.VertexShader->VerifyBoundUniformBufferParameters();
-	GlobalBoundShaderState.Get()->Args.PixelShader->VerifyBoundUniformBufferParameters();
-	GlobalBoundShaderState.Get()->Args.GeometryShader->VerifyBoundUniformBufferParameters();
+	WorkArea->Args.VertexShader->VerifyBoundUniformBufferParameters();
+	WorkArea->Args.PixelShader->VerifyBoundUniformBufferParameters();
+	WorkArea->Args.GeometryShader->VerifyBoundUniformBufferParameters();
 
-	FGlobalBoundShaderState_Internal* BSS = GlobalBoundShaderState.Get()->BSS;
+	FGlobalBoundShaderState_Internal* BSS = WorkArea->BSS;
 	bool bNewBSS = false;
 	if (!BSS)
 	{
@@ -116,15 +118,15 @@ static FBoundShaderStateRHIParamRef GetGlobalBoundShaderState_Internal(FGlobalBo
 		bNewBSS = true;
 	}
 	FBoundShaderStateRHIParamRef Result = BSS->GetInitializedRHI(
-		GlobalBoundShaderState.Get()->Args.VertexDeclarationRHI,
-		GETSAFERHISHADER_VERTEX(GlobalBoundShaderState.Get()->Args.VertexShader),
-		GETSAFERHISHADER_PIXEL(GlobalBoundShaderState.Get()->Args.PixelShader),
-		(FGeometryShaderRHIParamRef)GETSAFERHISHADER_GEOMETRY(GlobalBoundShaderState.Get()->Args.GeometryShader));
+		WorkArea->Args.VertexDeclarationRHI,
+		GETSAFERHISHADER_VERTEX(WorkArea->Args.VertexShader),
+		GETSAFERHISHADER_PIXEL(WorkArea->Args.PixelShader),
+		(FGeometryShaderRHIParamRef)GETSAFERHISHADER_GEOMETRY(WorkArea->Args.GeometryShader));
 	if (bNewBSS)
 	{
 		FPlatformMisc::MemoryBarrier();
 		// this probably doesn't need to be an atomic, assuming a pointer write is atomic, but this checks alignment, etc
-		FGlobalBoundShaderState_Internal* OldBSS = (FGlobalBoundShaderState_Internal*)FPlatformAtomics::InterlockedCompareExchangePointer((void**)&GlobalBoundShaderState.Get()->BSS, BSS, nullptr);
+		FGlobalBoundShaderState_Internal* OldBSS = (FGlobalBoundShaderState_Internal*)FPlatformAtomics::InterlockedCompareExchangePointer((void**)&WorkArea->BSS, BSS, nullptr);
 		check(!OldBSS); // should not be any concurrent writes
 	}
 	return Result;
@@ -156,12 +158,13 @@ public:
 
 	void DoTask(ENamedThreads::Type CurrentThread, const FGraphEventRef& MyCompletionGraphEvent)
 	{
-		RHICmdList.SetBoundShaderState(GetGlobalBoundShaderState_Internal(GlobalBoundShaderState));
+		RHICmdList.SetBoundShaderState(GetGlobalBoundShaderState_Internal(GlobalBoundShaderState, GRHIFeatureLevel));
 	}
 };
 
 void SetGlobalBoundShaderState(
 	FRHICommandList& RHICmdList,
+	ERHIFeatureLevel::Type FeatureLevel, 
 	FGlobalBoundShaderState& GlobalBoundShaderState,
 	FVertexDeclarationRHIParamRef VertexDeclarationRHI,
 	FShader* VertexShader,
@@ -169,7 +172,7 @@ void SetGlobalBoundShaderState(
 	FShader* GeometryShader
 	)
 {
-	FGlobalBoundShaderStateWorkArea* ExistingGlobalBoundShaderState = GlobalBoundShaderState.Get();
+	FGlobalBoundShaderStateWorkArea* ExistingGlobalBoundShaderState = GlobalBoundShaderState.Get(FeatureLevel);
 	if (!ExistingGlobalBoundShaderState)
 	{
 		FGlobalBoundShaderStateWorkArea* NewGlobalBoundShaderState = new FGlobalBoundShaderStateWorkArea();
@@ -179,18 +182,18 @@ void SetGlobalBoundShaderState(
 		NewGlobalBoundShaderState->Args.GeometryShader = GeometryShader;
 		FPlatformMisc::MemoryBarrier();
 
-		FGlobalBoundShaderStateWorkArea* OldGlobalBoundShaderState = (FGlobalBoundShaderStateWorkArea*)FPlatformAtomics::InterlockedCompareExchangePointer((void**)GlobalBoundShaderState.GetPtr(), NewGlobalBoundShaderState, nullptr);
+		FGlobalBoundShaderStateWorkArea* OldGlobalBoundShaderState = (FGlobalBoundShaderStateWorkArea*)FPlatformAtomics::InterlockedCompareExchangePointer((void**)GlobalBoundShaderState.GetPtr(FeatureLevel), NewGlobalBoundShaderState, nullptr);
 
 		if (OldGlobalBoundShaderState != nullptr)
 		{
 			//we lost
 			delete NewGlobalBoundShaderState;
-			check(OldGlobalBoundShaderState == GlobalBoundShaderState.Get());
+			check(OldGlobalBoundShaderState == GlobalBoundShaderState.Get(FeatureLevel));
 			ExistingGlobalBoundShaderState = OldGlobalBoundShaderState;
 		}
 		else
 		{
-			check(NewGlobalBoundShaderState == GlobalBoundShaderState.Get());
+			check(NewGlobalBoundShaderState == GlobalBoundShaderState.Get(FeatureLevel));
 			ExistingGlobalBoundShaderState = NewGlobalBoundShaderState;
 		}
 	}
@@ -209,7 +212,7 @@ void SetGlobalBoundShaderState(
 	if (RHICmdList.Bypass() || IsInRenderingThread())
 	{
 		// we can just do it as usual
-		RHICmdList.SetBoundShaderState(GetGlobalBoundShaderState_Internal(GlobalBoundShaderState));
+		RHICmdList.SetBoundShaderState(GetGlobalBoundShaderState_Internal(GlobalBoundShaderState, FeatureLevel));
 		return;
 	}
 	if (ExistingGlobalBoundShaderState->BSS)
diff --git a/Engine/Source/Runtime/Engine/Private/UnrealClient.cpp b/Engine/Source/Runtime/Engine/Private/UnrealClient.cpp
--- a/Engine/Source/Runtime/Engine/Private/UnrealClient.cpp
+++ b/Engine/Source/Runtime/Engine/Private/UnrealClient.cpp
@@ -689,7 +689,7 @@ class FDummyViewport : public FViewport
 public:
 	FDummyViewport(FViewportClient* InViewportClient)
 		: FViewport(InViewportClient)
-		, DebugCanvas( this, NULL, InViewportClient->GetWorld() )
+		, DebugCanvas(this, NULL, InViewportClient->GetWorld(), InViewportClient->GetWorld()->FeatureLevel)
 	{
 		DebugCanvas.SetAllowedModes(0);
 	}
@@ -805,7 +805,7 @@ void FViewport::HighResScreenshot()
 	ViewportClient->GetEngineShowFlags()->HighResScreenshotMask = GetHighResScreenshotConfig().bMaskEnabled ? 1 : 0;
 	ViewportClient->GetEngineShowFlags()->MotionBlur = 0;
 
-	FCanvas Canvas(DummyViewport, NULL, ViewportClient->GetWorld());
+	FCanvas Canvas(DummyViewport, NULL, ViewportClient->GetWorld(), ViewportClient->GetWorld()->FeatureLevel);
 	{
 		ViewportClient->Draw(DummyViewport, &Canvas);
 	}
@@ -1070,7 +1070,7 @@ void FViewport::Draw( bool bShouldPresent /*= true */)
 					GameThread.Waits = 0;
 				}
 
-				FCanvas Canvas(this, NULL, ViewportClient->GetWorld());
+				FCanvas Canvas(this, NULL, ViewportClient->GetWorld(), ViewportClient->GetWorld()->FeatureLevel);
 				{
 					ViewportClient->Draw(this, &Canvas);
 				}
@@ -1173,7 +1173,7 @@ const TArray<FColor>& FViewport::GetRawHitProxyData(FIntRect InRect)
 		});
 
 		// Let the viewport client draw its hit proxies.
-		FCanvas Canvas(&HitProxyMap, &HitProxyMap, ViewportClient->GetWorld());
+		FCanvas Canvas(&HitProxyMap, &HitProxyMap, ViewportClient->GetWorld(), ViewportClient->GetWorld()->FeatureLevel);
 		{
 			ViewportClient->Draw(this, &Canvas);
 		}
diff --git a/Engine/Source/Runtime/Engine/Private/UserInterface/Canvas.cpp b/Engine/Source/Runtime/Engine/Private/UserInterface/Canvas.cpp
--- a/Engine/Source/Runtime/Engine/Private/UserInterface/Canvas.cpp
+++ b/Engine/Source/Runtime/Engine/Private/UserInterface/Canvas.cpp
@@ -24,7 +24,7 @@ DEFINE_STAT(STAT_Canvas_GetBatchElementsTime);
 DEFINE_STAT(STAT_Canvas_AddTileRenderTime);
 DEFINE_STAT(STAT_Canvas_NumBatchesCreated);
 
-FCanvas::FCanvas(FRenderTarget* InRenderTarget,FHitProxyConsumer* InHitProxyConsumer, UWorld* InWorld)
+FCanvas::FCanvas(FRenderTarget* InRenderTarget, FHitProxyConsumer* InHitProxyConsumer, UWorld* InWorld, ERHIFeatureLevel::Type InFeatureLevel)
 :	ViewRect(0,0,0,0)
 ,	RenderTarget(InRenderTarget)
 ,	HitProxyConsumer(InHitProxyConsumer)
@@ -33,9 +33,11 @@ FCanvas::FCanvas(FRenderTarget* InRenderTarget,FHitProxyConsumer* InHitProxyCons
 ,	CurrentRealTime(0)
 ,	CurrentWorldTime(0)
 ,	CurrentDeltaWorldTime(0)
+,	FeatureLevel(InFeatureLevel)
 {
 	Construct();
 
+	check(InWorld);
 	if (InWorld)
 	{
 		CurrentRealTime = InWorld->GetRealTimeSeconds();
@@ -44,7 +46,7 @@ FCanvas::FCanvas(FRenderTarget* InRenderTarget,FHitProxyConsumer* InHitProxyCons
 	}
 }
 
-FCanvas::FCanvas(FRenderTarget* InRenderTarget,FHitProxyConsumer* InHitProxyConsumer, float InRealTime, float InWorldTime, float InWorldDeltaTime)
+FCanvas::FCanvas(FRenderTarget* InRenderTarget,FHitProxyConsumer* InHitProxyConsumer, float InRealTime, float InWorldTime, float InWorldDeltaTime, ERHIFeatureLevel::Type InFeatureLevel)
 :	ViewRect(0,0,0,0)
 ,	RenderTarget(InRenderTarget)
 ,	HitProxyConsumer(InHitProxyConsumer)
@@ -53,6 +55,7 @@ FCanvas::FCanvas(FRenderTarget* InRenderTarget,FHitProxyConsumer* InHitProxyCons
 ,	CurrentRealTime(InRealTime)
 ,	CurrentWorldTime(InWorldTime)
 ,	CurrentDeltaWorldTime(InWorldDeltaTime)
+,	FeatureLevel(InFeatureLevel)
 {
 	Construct();
 }
@@ -176,6 +179,7 @@ bool FCanvasBatchedElementRenderItem::Render_RenderThread(FRHICommandListImmedia
 		// draw batched items
 		Data->BatchedElements.Draw(
 			RHICmdList,
+			Canvas->GetFeatureLevel(),
 			bNeedsToSwitchVerticalAxis,
 			Data->Transform.GetMatrix(),
 			CanvasRenderTarget->GetSizeXY().X,
@@ -197,7 +201,7 @@ bool FCanvasBatchedElementRenderItem::Render_RenderThread(FRHICommandListImmedia
 	return bDirty;
 }
 
-bool FCanvasBatchedElementRenderItem::Render_GameThread(const FCanvas* Canvas )
+bool FCanvasBatchedElementRenderItem::Render_GameThread(const FCanvas* Canvas)
 {	
 	checkSlow(Data);
 	bool bDirty=false;		
@@ -222,6 +226,8 @@ bool FCanvasBatchedElementRenderItem::Render_GameThread(const FCanvas* Canvas )
 			uint32 ViewportSizeY;
 			float DisplayGamma;
 			uint32 AllowedCanvasModes;
+			ERHIFeatureLevel::Type FeatureLevel;
+			EShaderPlatform ShaderPlatform;
 		};
 		// all the parameters needed for rendering
 		FBatchedDrawParameters DrawParameters =
@@ -231,7 +237,9 @@ bool FCanvasBatchedElementRenderItem::Render_GameThread(const FCanvas* Canvas )
 			(uint32)CanvasRenderTarget->GetSizeXY().X,
 			(uint32)CanvasRenderTarget->GetSizeXY().Y,
 			Gamma,
-			Canvas->GetAllowedModes()
+			Canvas->GetAllowedModes(),
+			Canvas->GetFeatureLevel(),
+			Canvas->GetShaderPlatform()
 		};
 		ENQUEUE_UNIQUE_RENDER_COMMAND_ONEPARAMETER(
 			BatchedDrawCommand,
@@ -242,6 +250,7 @@ bool FCanvasBatchedElementRenderItem::Render_GameThread(const FCanvas* Canvas )
 			// draw batched items
 			Parameters.RenderData->BatchedElements.Draw(
 				RHICmdList,
+				Parameters.FeatureLevel,
 				bNeedsToSwitchVerticalAxis,
 				Parameters.RenderData->Transform.GetMatrix(),
 				Parameters.ViewportSizeX,
diff --git a/Engine/Source/Runtime/Engine/Private/World.cpp b/Engine/Source/Runtime/Engine/Private/World.cpp
--- a/Engine/Source/Runtime/Engine/Private/World.cpp
+++ b/Engine/Source/Runtime/Engine/Private/World.cpp
@@ -5168,8 +5168,9 @@ void UWorld::ChangeFeatureLevel(ERHIFeatureLevel::Type InFeatureLevel)
 
 		FGlobalComponentReregisterContext RecreateComponents;
 
-		GCurrentRHIFeatureLevel = InFeatureLevel;
-		GRHIShaderPlatform = GShaderPlatformForFeatureLevel[InFeatureLevel];
+		// TODO: Remove these
+		GRHIFeatureLevelValue = InFeatureLevel;
+		GRHIShaderPlatformValue = GShaderPlatformForFeatureLevel[InFeatureLevel];
 
 		UMaterial::AllMaterialsCacheResourceShadersForRendering();
 		UMaterialInstance::AllMaterialsCacheResourceShadersForRendering();
@@ -5211,8 +5212,8 @@ void UWorld::ChangeAllWorldFeatureLevels(ERHIFeatureLevel::Type InFeatureLevel)
 
 		FGlobalComponentReregisterContext RecreateComponents;
 
-		GCurrentRHIFeatureLevel = InFeatureLevel;
-		GRHIShaderPlatform = GShaderPlatformForFeatureLevel[InFeatureLevel];
+		GRHIFeatureLevelValue = InFeatureLevel;
+		GRHIShaderPlatformValue = GShaderPlatformForFeatureLevel[InFeatureLevel];
 
 		UMaterial::AllMaterialsCacheResourceShadersForRendering();
 		UMaterialInstance::AllMaterialsCacheResourceShadersForRendering();
diff --git a/Engine/Source/Runtime/Engine/Public/BatchedElements.h b/Engine/Source/Runtime/Engine/Public/BatchedElements.h
--- a/Engine/Source/Runtime/Engine/Public/BatchedElements.h
+++ b/Engine/Source/Runtime/Engine/Public/BatchedElements.h
@@ -83,7 +83,7 @@ class FBatchedElementParameters
 public:
 
 	/** Binds vertex and pixel shaders for this element */
-	virtual void BindShaders(FRHICommandList& RHICmdList, const FMatrix& InTransform, const float InGamma, const FMatrix& ColorWeights, const FTexture* Texture) = 0;
+	virtual void BindShaders(FRHICommandList& RHICmdList, ERHIFeatureLevel::Type InFeatureLevel, const FMatrix& InTransform, const float InGamma, const FMatrix& ColorWeights, const FTexture* Texture) = 0;
 
 };
 
@@ -176,7 +176,7 @@ public:
 	 * @param View			Optional FSceneView for shaders that need access to view constants
 	 * @param DepthTexture	DepthTexture for manual depth testing with editor compositing in the pixel shader
 	 */
-	bool Draw(FRHICommandList& RHICmdList, bool bNeedToSwitchVerticalAxis, const FMatrix& Transform, uint32 ViewportSizeX, uint32 ViewportSizeY, bool bHitTesting, float Gamma = 1.0f, const FSceneView* View = NULL, FTexture2DRHIRef DepthTexture = FTexture2DRHIRef(), EBlendModeFilter::Type Filter = EBlendModeFilter::All) const;
+	bool Draw(FRHICommandList& RHICmdList, ERHIFeatureLevel::Type FeatureLevel, bool bNeedToSwitchVerticalAxis, const FMatrix& Transform, uint32 ViewportSizeX, uint32 ViewportSizeY, bool bHitTesting, float Gamma = 1.0f, const FSceneView* View = NULL, FTexture2DRHIRef DepthTexture = FTexture2DRHIRef(), EBlendModeFilter::Type Filter = EBlendModeFilter::All) const;
 	
 	FORCEINLINE bool HasPrimsToDraw() const
 	{
@@ -312,6 +312,7 @@ private:
 	 */
 	void PrepareShaders(
 		FRHICommandList& RHICmdList,
+		ERHIFeatureLevel::Type FeatureLevel,
 		ESimpleElementBlendMode BlendMode,
 		const FMatrix& Transform,
 		bool bSwitchVerticalAxis,
diff --git a/Engine/Source/Runtime/Engine/Public/CanvasTypes.h b/Engine/Source/Runtime/Engine/Public/CanvasTypes.h
--- a/Engine/Source/Runtime/Engine/Public/CanvasTypes.h
+++ b/Engine/Source/Runtime/Engine/Public/CanvasTypes.h
@@ -47,12 +47,12 @@ public:
 	/** 
 	* Constructor.
 	*/
-	ENGINE_API FCanvas(FRenderTarget* InRenderTarget,FHitProxyConsumer* InHitProxyConsumer, UWorld* InWorld);
+	ENGINE_API FCanvas(FRenderTarget* InRenderTarget,FHitProxyConsumer* InHitProxyConsumer, UWorld* InWorld, ERHIFeatureLevel::Type InFeatureLevel);
 
 	/** 
 	* Constructor. For situations where a world is not available, but time information is
 	*/
-	ENGINE_API FCanvas(FRenderTarget* InRenderTarget,FHitProxyConsumer* InHitProxyConsumer, float InRealTime, float InWorldTime, float InWorldDeltaTime);
+	ENGINE_API FCanvas(FRenderTarget* InRenderTarget, FHitProxyConsumer* InHitProxyConsumer, float InRealTime, float InWorldTime, float InWorldDeltaTime, ERHIFeatureLevel::Type InFeatureLevel);
 
 	/** 
 	* Destructor.
@@ -298,6 +298,20 @@ public:
 	 */
 	bool HasBatchesToRender() const;
 
+	/**
+	* Access current feature level
+	*
+	* @return feature level that this canvas is rendering at
+	*/
+	ERHIFeatureLevel::Type GetFeatureLevel() const { return FeatureLevel; }
+
+	/**
+	* Access current shader platform
+	*
+	* @return shader platform that this canvas is rendering at
+	*/
+	EShaderPlatform GetShaderPlatform() const { return ShaderPlatform; }
+
 public:
 	float AlphaModulate;
 
@@ -368,6 +382,10 @@ private:
 	float CurrentDeltaWorldTime;
 	/** true, if Canvas should be scaled to whole render target */
 	bool bScaledToRenderTarget;
+	/** Feature level that we are currently rendering with */
+	ERHIFeatureLevel::Type FeatureLevel;
+	/** Shader platform that we are currently rendering with */
+	EShaderPlatform ShaderPlatform;
 
 	/** 
 	* Shared construction function
diff --git a/Engine/Source/Runtime/Engine/Public/SceneView.h b/Engine/Source/Runtime/Engine/Public/SceneView.h
--- a/Engine/Source/Runtime/Engine/Public/SceneView.h
+++ b/Engine/Source/Runtime/Engine/Public/SceneView.h
@@ -452,6 +452,9 @@ public:
 
 	/** Get the feature level for this view **/
 	ERHIFeatureLevel::Type GetFeatureLevel() const;
+
+	/** Get the feature level for this view **/
+	EShaderPlatform GetShaderPlatform() const;
 };
 
 //////////////////////////////////////////////////////////////////////////
diff --git a/Engine/Source/Runtime/Engine/Public/StaticBoundShaderState.h b/Engine/Source/Runtime/Engine/Public/StaticBoundShaderState.h
--- a/Engine/Source/Runtime/Engine/Public/StaticBoundShaderState.h
+++ b/Engine/Source/Runtime/Engine/Public/StaticBoundShaderState.h
@@ -80,8 +80,8 @@ struct FGlobalBoundShaderState
 {
 public:
 
-	FGlobalBoundShaderStateWorkArea* Get(ERHIFeatureLevel::Type InFeatureLevel = GRHIFeatureLevel)  { return WorkAreas[InFeatureLevel]; }
-	FGlobalBoundShaderStateWorkArea** GetPtr(ERHIFeatureLevel::Type InFeatureLevel = GRHIFeatureLevel)  { return &WorkAreas[InFeatureLevel]; }
+	FGlobalBoundShaderStateWorkArea* Get(ERHIFeatureLevel::Type InFeatureLevel)  { return WorkAreas[InFeatureLevel]; }
+	FGlobalBoundShaderStateWorkArea** GetPtr(ERHIFeatureLevel::Type InFeatureLevel)  { return &WorkAreas[InFeatureLevel]; }
 
 private:
 
@@ -103,6 +103,7 @@ private:
 
 ENGINE_API void SetGlobalBoundShaderState(
 	FRHICommandList& RHICmdList,
+	ERHIFeatureLevel::Type FeatureLevel,
 	FGlobalBoundShaderState& BoundShaderState,
 	FVertexDeclarationRHIParamRef VertexDeclaration,
 	FShader* VertexShader,
diff --git a/Engine/Source/Runtime/GameLiveStreaming/GameLiveStreaming.cpp b/Engine/Source/Runtime/GameLiveStreaming/GameLiveStreaming.cpp
--- a/Engine/Source/Runtime/GameLiveStreaming/GameLiveStreaming.cpp
+++ b/Engine/Source/Runtime/GameLiveStreaming/GameLiveStreaming.cpp
@@ -333,7 +333,7 @@ void FGameLiveStreaming::StartCopyingNextGameVideoFrame( const FViewportRHIRef&
 		TShaderMapRef<FScreenPS> PixelShader(GetGlobalShaderMap());
 
 		static FGlobalBoundShaderState BoundShaderState;
-		SetGlobalBoundShaderState(RHICmdList, BoundShaderState, Context.RendererModule->GetFilterVertexDeclaration().VertexDeclarationRHI, *VertexShader, *PixelShader);
+		SetGlobalBoundShaderState(RHICmdList, GRHIFeatureLevel, BoundShaderState, Context.RendererModule->GetFilterVertexDeclaration().VertexDeclarationRHI, *VertexShader, *PixelShader);
 
 		if( Context.ResizeTo != FIntPoint( ViewportBackBuffer->GetSizeX(), ViewportBackBuffer->GetSizeY() ) )
 		{
diff --git a/Engine/Source/Runtime/IOS/MetalRHI/Private/MetalRHI.cpp b/Engine/Source/Runtime/IOS/MetalRHI/Private/MetalRHI.cpp
--- a/Engine/Source/Runtime/IOS/MetalRHI/Private/MetalRHI.cpp
+++ b/Engine/Source/Runtime/IOS/MetalRHI/Private/MetalRHI.cpp
@@ -42,9 +42,9 @@ FMetalDynamicRHI::FMetalDynamicRHI()
 	check( !GIsThreadedRendering );
 
 	// Initialize the RHI capabilities.
-	GRHIFeatureLevel = ERHIFeatureLevel::ES2;
-	GRHIShaderPlatform = SP_METAL;
-	GMaxRHIFeatureLevel = GCurrentRHIFeatureLevel = ERHIFeatureLevel::ES2;
+	GRHIFeatureLevelValue = ERHIFeatureLevel::ES2;
+	GRHIShaderPlatformValue = SP_METAL;
+	GMaxRHIFeatureLevelValue = GRHIFeatureLevelValue = ERHIFeatureLevel::ES2;
 
 	GPixelCenterOffset = 0.0f;
 	GSupportsVertexInstancing = false;
diff --git a/Engine/Source/Runtime/NullDrv/Private/NullRHI.cpp b/Engine/Source/Runtime/NullDrv/Private/NullRHI.cpp
--- a/Engine/Source/Runtime/NullDrv/Private/NullRHI.cpp
+++ b/Engine/Source/Runtime/NullDrv/Private/NullRHI.cpp
@@ -14,7 +14,7 @@
 
 FNullDynamicRHI::FNullDynamicRHI()
 {
-	GRHIShaderPlatform = ShaderFormatToLegacyShaderPlatform(FName(FPlatformMisc::GetNullRHIShaderFormat()));
+	GRHIShaderPlatformValue = ShaderFormatToLegacyShaderPlatform(FName(FPlatformMisc::GetNullRHIShaderFormat()));
 }
 
 void FNullDynamicRHI::Init()
diff --git a/Engine/Source/Runtime/OpenGLDrv/Private/OpenGLDevice.cpp b/Engine/Source/Runtime/OpenGLDrv/Private/OpenGLDevice.cpp
--- a/Engine/Source/Runtime/OpenGLDrv/Private/OpenGLDevice.cpp
+++ b/Engine/Source/Runtime/OpenGLDrv/Private/OpenGLDevice.cpp
@@ -606,8 +606,8 @@ static void InitRHICapabilitiesForGL()
 	const GLint MinorVersion = FOpenGL::GetMinorVersion();
 
 	// Shader platform & RHI feature level
-	GMaxRHIFeatureLevel = GCurrentRHIFeatureLevel = FOpenGL::GetFeatureLevel();
-	GRHIShaderPlatform = FOpenGL::GetShaderPlatform();
+	GMaxRHIFeatureLevelValue = GRHIFeatureLevelValue = FOpenGL::GetFeatureLevel();
+	GRHIShaderPlatformValue = FOpenGL::GetShaderPlatform();
 
 	FString FeatureLevelName;
 	GetFeatureLevelName(GMaxRHIFeatureLevel, FeatureLevelName);
diff --git a/Engine/Source/Runtime/OpenGLDrv/Private/OpenGLDrv.cpp b/Engine/Source/Runtime/OpenGLDrv/Private/OpenGLDrv.cpp
--- a/Engine/Source/Runtime/OpenGLDrv/Private/OpenGLDrv.cpp
+++ b/Engine/Source/Runtime/OpenGLDrv/Private/OpenGLDrv.cpp
@@ -413,7 +413,7 @@ void FOpenGLDynamicRHI::IssueLongGPUTask()
 		TShaderMapRef<TOneColorVS<true> > VertexShader(GetGlobalShaderMap());
 		TShaderMapRef<FOpenGLRHILongGPUTaskPS> PixelShader(GetGlobalShaderMap());
 
-		SetGlobalBoundShaderState(RHICmdList, LongGPUTaskBoundShaderState, GOpenGLVector4VertexDeclaration.VertexDeclarationRHI, *VertexShader, *PixelShader, 0);
+		SetGlobalBoundShaderState(RHICmdList, GMaxRHIFeatureLevel, LongGPUTaskBoundShaderState, GOpenGLVector4VertexDeclaration.VertexDeclarationRHI, *VertexShader, *PixelShader, 0);
 
 		// Draw a fullscreen quad
 		FVector4 Vertices[4];
diff --git a/Engine/Source/Runtime/RHI/Private/RHI.cpp b/Engine/Source/Runtime/RHI/Private/RHI.cpp
--- a/Engine/Source/Runtime/RHI/Private/RHI.cpp
+++ b/Engine/Source/Runtime/RHI/Private/RHI.cpp
@@ -217,17 +217,28 @@ void RHIPrivateBeginFrame()
 // The current shader platform.
 //
 
-RHI_API EShaderPlatform GRHIShaderPlatform = SP_PCD3D_SM5;
+RHI_API EShaderPlatform GRHIShaderPlatformValue = SP_PCD3D_SM5;
 
 /** The maximum feature level supported on this machine */
-ERHIFeatureLevel::Type GMaxRHIFeatureLevel = ERHIFeatureLevel::SM5;
+RHI_API ERHIFeatureLevel::Type GMaxRHIFeatureLevelValue = ERHIFeatureLevel::SM5;
 
 /** The current feature level being used on this machine - This is only temporary and will go away when feature level becomes dynamic */
-ERHIFeatureLevel::Type GCurrentRHIFeatureLevel = ERHIFeatureLevel::SM5;
+RHI_API ERHIFeatureLevel::Type GRHIFeatureLevelValue = ERHIFeatureLevel::SM5;
 
 RHI_API ERHIFeatureLevel::Type GetCurrentRHIFeatureLevel()
 {
-	return GCurrentRHIFeatureLevel;
+	//check(!IsInRenderingThread());
+	return GRHIFeatureLevelValue;
+}
+
+RHI_API ERHIFeatureLevel::Type GetMaxRHIFeatureLevel()
+{
+	return GMaxRHIFeatureLevelValue;
+}
+
+RHI_API EShaderPlatform GetCurrentRHIShaderPlatform()
+{
+	return GRHIShaderPlatformValue;
 }
 
 FName FeatureLevelNames[] = 
diff --git a/Engine/Source/Runtime/RHI/Public/RHI.h b/Engine/Source/Runtime/RHI/Public/RHI.h
--- a/Engine/Source/Runtime/RHI/Public/RHI.h
+++ b/Engine/Source/Runtime/RHI/Public/RHI.h
@@ -223,7 +223,7 @@ inline FMatrix AdjustProjectionMatrixForRHI(const FMatrix& InProjectionMatrix)
 }
 
 /** Current shader platform. */
-extern RHI_API EShaderPlatform GRHIShaderPlatform;
+
 
 /** Finds a corresponding ERHIFeatureLevel::Type given an FName, or returns false if one could not be found. */
 extern RHI_API bool GetFeatureLevelFromName(FName Name, ERHIFeatureLevel::Type& OutFeatureLevel);
@@ -234,19 +234,26 @@ extern RHI_API void GetFeatureLevelName(ERHIFeatureLevel::Type InFeatureLevel, F
 /** Creates an FName for the given feature level. */
 extern RHI_API void GetFeatureLevelName(ERHIFeatureLevel::Type InFeatureLevel, FName& OutName);
 
-extern RHI_API ERHIFeatureLevel::Type GMaxRHIFeatureLevel;
-extern RHI_API ERHIFeatureLevel::Type GCurrentRHIFeatureLevel;
+extern RHI_API ERHIFeatureLevel::Type GRHIFeatureLevelValue;
+extern RHI_API ERHIFeatureLevel::Type GMaxRHIFeatureLevelValue;
+extern RHI_API EShaderPlatform GRHIShaderPlatformValue;
 
 /** Function for retrieving the current feature level. Only exists to provide a convient way of tracking access to this global during mobile preview work */
 extern RHI_API ERHIFeatureLevel::Type GetCurrentRHIFeatureLevel();
+extern RHI_API ERHIFeatureLevel::Type GetMaxRHIFeatureLevel();
+extern RHI_API EShaderPlatform GetCurrentRHIShaderPlatform();
 
 /** treating GRHIFeatureLevel as a function allows for better usage tracking at small cost to performance. */
 #define RHI_FEATURE_LEVEL_AS_FUNCTION 0
 
 #if RHI_FEATURE_LEVEL_AS_FUNCTION
 	#define GRHIFeatureLevel GetCurrentRHIFeatureLevel()
+	#define GMaxRHIFeatureLevel GetMaxRHIFeatureLevel()
+	#define GRHIShaderPlatform GetCurrentRHIShaderPlatform()
 #else
-	#define GRHIFeatureLevel GCurrentRHIFeatureLevel
+	#define GRHIFeatureLevel GRHIFeatureLevelValue
+	#define GMaxRHIFeatureLevel GMaxRHIFeatureLevelValue
+	#define GRHIShaderPlatform GRHIShaderPlatformValue
 #endif
 
 /** Table for finding out which shader platform corresponds to a given feature level for this RHI. */
diff --git a/Engine/Source/Runtime/Renderer/Private/AtmosphereRendering.cpp b/Engine/Source/Runtime/Renderer/Private/AtmosphereRendering.cpp
--- a/Engine/Source/Runtime/Renderer/Private/AtmosphereRendering.cpp
+++ b/Engine/Source/Runtime/Renderer/Private/AtmosphereRendering.cpp
@@ -384,7 +384,7 @@ void SetAtmosphericFogShaders(FRHICommandList& RHICmdList, FScene* Scene, const
 		break;
 	}
 	
-	SetGlobalBoundShaderState(RHICmdList, AtmosphereBoundShaderState[Scene->AtmosphericFog->RenderFlag], GAtmophereVertexDeclaration.VertexDeclarationRHI, *VertexShader, PixelShader);
+	SetGlobalBoundShaderState(RHICmdList, View.GetFeatureLevel(), AtmosphereBoundShaderState[Scene->AtmosphericFog->RenderFlag], GAtmophereVertexDeclaration.VertexDeclarationRHI, *VertexShader, PixelShader);
 	VertexShader->SetParameters(RHICmdList, View);
 	PixelShader->SetParameters(RHICmdList, View, LightShaftOcclusion);
 }
@@ -1095,7 +1095,7 @@ void FAtmosphericFogSceneInfo::RenderAtmosphereShaders(FRHICommandList& RHICmdLi
 			TShaderMapRef<FAtmosphereTransmittancePS> PixelShader(GetGlobalShaderMap());
 
 			static FGlobalBoundShaderState BoundShaderState;
-			SetGlobalBoundShaderState(RHICmdList, BoundShaderState, GFilterVertexDeclaration.VertexDeclarationRHI, *VertexShader, *PixelShader);
+			SetGlobalBoundShaderState(RHICmdList, View.GetFeatureLevel(), BoundShaderState, GFilterVertexDeclaration.VertexDeclarationRHI, *VertexShader, *PixelShader);
 			PixelShader->SetParameters(RHICmdList, View);
 			DrawQuad(RHICmdList, ViewRect, *VertexShader);
 
@@ -1111,7 +1111,7 @@ void FAtmosphericFogSceneInfo::RenderAtmosphereShaders(FRHICommandList& RHICmdLi
 			TShaderMapRef<FAtmosphereIrradiance1PS> PixelShader(GetGlobalShaderMap());
 
 			static FGlobalBoundShaderState BoundShaderState;
-			SetGlobalBoundShaderState(RHICmdList, BoundShaderState, GFilterVertexDeclaration.VertexDeclarationRHI, *VertexShader, *PixelShader);
+			SetGlobalBoundShaderState(RHICmdList, View.GetFeatureLevel(), BoundShaderState, GFilterVertexDeclaration.VertexDeclarationRHI, *VertexShader, *PixelShader);
 
 			PixelShader->SetParameters(RHICmdList, AtmosphereTextures);
 
@@ -1135,7 +1135,7 @@ void FAtmosphericFogSceneInfo::RenderAtmosphereShaders(FRHICommandList& RHICmdLi
 
 				static FGlobalBoundShaderState BoundShaderState;
 				
-				SetGlobalBoundShaderState(RHICmdList, BoundShaderState, GFilterVertexDeclaration.VertexDeclarationRHI, *VertexShader, *PixelShader, *GeometryShader);
+				SetGlobalBoundShaderState(RHICmdList, View.GetFeatureLevel(), BoundShaderState, GFilterVertexDeclaration.VertexDeclarationRHI, *VertexShader, *PixelShader, *GeometryShader);
 
 				//
 				float r;
@@ -1177,7 +1177,7 @@ void FAtmosphericFogSceneInfo::RenderAtmosphereShaders(FRHICommandList& RHICmdLi
 
 				static FGlobalBoundShaderState BoundShaderState;
 				
-				SetGlobalBoundShaderState(RHICmdList, BoundShaderState, GFilterVertexDeclaration.VertexDeclarationRHI, *VertexShader, *PixelShader, *GeometryShader);
+				SetGlobalBoundShaderState(RHICmdList, View.GetFeatureLevel(), BoundShaderState, GFilterVertexDeclaration.VertexDeclarationRHI, *VertexShader, *PixelShader, *GeometryShader);
 
 				float r;
 				FVector4 DhdH;
@@ -1206,7 +1206,7 @@ void FAtmosphericFogSceneInfo::RenderAtmosphereShaders(FRHICommandList& RHICmdLi
 
 				static FGlobalBoundShaderState BoundShaderState;
 				
-				SetGlobalBoundShaderState(RHICmdList, BoundShaderState, GFilterVertexDeclaration.VertexDeclarationRHI, *VertexShader, *PixelShader, *GeometryShader);
+				SetGlobalBoundShaderState(RHICmdList, View.GetFeatureLevel(), BoundShaderState, GFilterVertexDeclaration.VertexDeclarationRHI, *VertexShader, *PixelShader, *GeometryShader);
 
 				float r;
 				FVector4 DhdH;
@@ -1232,7 +1232,7 @@ void FAtmosphericFogSceneInfo::RenderAtmosphereShaders(FRHICommandList& RHICmdLi
 
 			static FGlobalBoundShaderState BoundShaderState;
 			
-			SetGlobalBoundShaderState(RHICmdList, BoundShaderState, GFilterVertexDeclaration.VertexDeclarationRHI, *VertexShader, *PixelShader);
+			SetGlobalBoundShaderState(RHICmdList, View.GetFeatureLevel(), BoundShaderState, GFilterVertexDeclaration.VertexDeclarationRHI, *VertexShader, *PixelShader);
 
 			PixelShader->SetParameters(RHICmdList, View, AtmoshpereOrder == 2 ? 1.f : 0.f, AtmosphereTextures);
 
@@ -1255,7 +1255,7 @@ void FAtmosphericFogSceneInfo::RenderAtmosphereShaders(FRHICommandList& RHICmdLi
 
 				static FGlobalBoundShaderState BoundShaderState;
 				
-				SetGlobalBoundShaderState(RHICmdList, BoundShaderState, GFilterVertexDeclaration.VertexDeclarationRHI, *VertexShader, *PixelShader, *GeometryShader);
+				SetGlobalBoundShaderState(RHICmdList, View.GetFeatureLevel(), BoundShaderState, GFilterVertexDeclaration.VertexDeclarationRHI, *VertexShader, *PixelShader, *GeometryShader);
 
 				float r;
 				FVector4 DhdH;
@@ -1283,7 +1283,7 @@ void FAtmosphericFogSceneInfo::RenderAtmosphereShaders(FRHICommandList& RHICmdLi
 
 			static FGlobalBoundShaderState BoundShaderState;
 			
-			SetGlobalBoundShaderState(RHICmdList, BoundShaderState, GFilterVertexDeclaration.VertexDeclarationRHI, *VertexShader, *PixelShader);
+			SetGlobalBoundShaderState(RHICmdList, View.GetFeatureLevel(), BoundShaderState, GFilterVertexDeclaration.VertexDeclarationRHI, *VertexShader, *PixelShader);
 
 			PixelShader->SetParameters(RHICmdList, AtmosphereTextures);
 
@@ -1310,7 +1310,7 @@ void FAtmosphericFogSceneInfo::RenderAtmosphereShaders(FRHICommandList& RHICmdLi
 
 				static FGlobalBoundShaderState BoundShaderState;
 				
-				SetGlobalBoundShaderState(RHICmdList, BoundShaderState, GFilterVertexDeclaration.VertexDeclarationRHI, *VertexShader, *PixelShader, *GeometryShader);
+				SetGlobalBoundShaderState(RHICmdList, View.GetFeatureLevel(), BoundShaderState, GFilterVertexDeclaration.VertexDeclarationRHI, *VertexShader, *PixelShader, *GeometryShader);
 
 				float r;
 				FVector4 DhdH;
@@ -1342,7 +1342,7 @@ void FAtmosphericFogSceneInfo::RenderAtmosphereShaders(FRHICommandList& RHICmdLi
 
 				static FGlobalBoundShaderState BoundShaderState;
 				
-				SetGlobalBoundShaderState(RHICmdList, BoundShaderState, GFilterVertexDeclaration.VertexDeclarationRHI, *VertexShader, *PixelShader, *GeometryShader);
+				SetGlobalBoundShaderState(RHICmdList, View.GetFeatureLevel(), BoundShaderState, GFilterVertexDeclaration.VertexDeclarationRHI, *VertexShader, *PixelShader, *GeometryShader);
 
 				float r;
 				FVector4 DhdH;
@@ -1371,7 +1371,7 @@ void FAtmosphericFogSceneInfo::RenderAtmosphereShaders(FRHICommandList& RHICmdLi
 
 				static FGlobalBoundShaderState BoundShaderState;
 				
-				SetGlobalBoundShaderState(RHICmdList, BoundShaderState, GFilterVertexDeclaration.VertexDeclarationRHI, *VertexShader, *PixelShader, *GeometryShader);
+				SetGlobalBoundShaderState(RHICmdList, View.GetFeatureLevel(), BoundShaderState, GFilterVertexDeclaration.VertexDeclarationRHI, *VertexShader, *PixelShader, *GeometryShader);
 
 				float r;
 				FVector4 DhdH;
diff --git a/Engine/Source/Runtime/Renderer/Private/BasePassRendering.cpp b/Engine/Source/Runtime/Renderer/Private/BasePassRendering.cpp
--- a/Engine/Source/Runtime/Renderer/Private/BasePassRendering.cpp
+++ b/Engine/Source/Runtime/Renderer/Private/BasePassRendering.cpp
@@ -187,6 +187,7 @@ public:
 				StaticMesh->VertexFactory,
 				StaticMesh->MaterialRenderProxy,
 				*Parameters.Material,
+				Parameters.FeatureLevel,
 				LightMapPolicy,
 				Parameters.BlendMode,
 				Parameters.TextureMode,
@@ -288,6 +289,7 @@ public:
 			Parameters.Mesh.VertexFactory,
 			Parameters.Mesh.MaterialRenderProxy,
 			*Parameters.Material,
+			Parameters.FeatureLevel,
 			LightMapPolicy,
 			Parameters.BlendMode,
 			Parameters.TextureMode,
diff --git a/Engine/Source/Runtime/Renderer/Private/BasePassRendering.h b/Engine/Source/Runtime/Renderer/Private/BasePassRendering.h
--- a/Engine/Source/Runtime/Renderer/Private/BasePassRendering.h
+++ b/Engine/Source/Runtime/Renderer/Private/BasePassRendering.h
@@ -475,6 +475,7 @@ public:
 		const FVertexFactory* InVertexFactory,
 		const FMaterialRenderProxy* InMaterialRenderProxy,
 		const FMaterial& InMaterialResource,
+		ERHIFeatureLevel::Type InFeatureLevel,
 		LightMapPolicyType InLightMapPolicy,
 		EBlendMode InBlendMode,
 		ESceneRenderTargetsMode::Type InSceneTextureMode,
@@ -498,7 +499,7 @@ public:
 	
 		const EMaterialTessellationMode MaterialTessellationMode = InMaterialResource.GetTessellationMode();
 
-		if (RHISupportsTessellation(GRHIShaderPlatform)
+		if (RHISupportsTessellation(GShaderPlatformForFeatureLevel[InFeatureLevel])
 			&& InVertexFactory->GetType()->SupportsTessellationShaders() 
 			&& MaterialTessellationMode != MTM_NoTessellation)
 		{
diff --git a/Engine/Source/Runtime/Renderer/Private/ClearQuad.cpp b/Engine/Source/Runtime/Renderer/Private/ClearQuad.cpp
--- a/Engine/Source/Runtime/Renderer/Private/ClearQuad.cpp
+++ b/Engine/Source/Runtime/Renderer/Private/ClearQuad.cpp
@@ -6,7 +6,7 @@
 FGlobalBoundShaderState GClearMRTBoundShaderState[8];
 
 // TODO support ExcludeRect
-void DrawClearQuadMRT(FRHICommandList& RHICmdList, bool bClearColor, int32 NumClearColors, const FLinearColor* ClearColorArray, bool bClearDepth, float Depth, bool bClearStencil, uint32 Stencil)
+void DrawClearQuadMRT(FRHICommandList& RHICmdList, ERHIFeatureLevel::Type FeatureLevel, bool bClearColor, int32 NumClearColors, const FLinearColor* ClearColorArray, bool bClearDepth, float Depth, bool bClearStencil, uint32 Stencil)
 {
 	// Set new states
 	FBlendStateRHIParamRef BlendStateRHI;
@@ -95,7 +95,7 @@ void DrawClearQuadMRT(FRHICommandList& RHICmdList, bool bClearColor, int32 NumCl
 		PixelShader = *MRTPixelShader;
 	}
 
-	SetGlobalBoundShaderState(RHICmdList, GClearMRTBoundShaderState[FMath::Max(NumClearColors - 1, 0)], GetVertexDeclarationFVector4(), *VertexShader, PixelShader);
+	SetGlobalBoundShaderState(RHICmdList, FeatureLevel, GClearMRTBoundShaderState[FMath::Max(NumClearColors - 1, 0)], GetVertexDeclarationFVector4(), *VertexShader, PixelShader);
 	FLinearColor ShaderClearColors[MaxSimultaneousRenderTargets];
 	FMemory::MemZero(ShaderClearColors);
 
diff --git a/Engine/Source/Runtime/Renderer/Private/ClearQuad.h b/Engine/Source/Runtime/Renderer/Private/ClearQuad.h
--- a/Engine/Source/Runtime/Renderer/Private/ClearQuad.h
+++ b/Engine/Source/Runtime/Renderer/Private/ClearQuad.h
@@ -3,9 +3,9 @@
 #pragma once
 
 // TODO support ExcludeRect
-void DrawClearQuadMRT(FRHICommandList& RHICmdList, bool bClearColor, int32 NumClearColors, const FLinearColor* ClearColorArray, bool bClearDepth, float Depth, bool bClearStencil, uint32 Stencil);
+void DrawClearQuadMRT(FRHICommandList& RHICmdList, ERHIFeatureLevel::Type FeatureLevel, bool bClearColor, int32 NumClearColors, const FLinearColor* ClearColorArray, bool bClearDepth, float Depth, bool bClearStencil, uint32 Stencil);
 
-inline void DrawClearQuad(FRHICommandList& RHICmdList, bool bClearColor, const FLinearColor& Color, bool bClearDepth, float Depth, bool bClearStencil, uint32 Stencil)
+inline void DrawClearQuad(FRHICommandList& RHICmdList, ERHIFeatureLevel::Type FeatureLevel, bool bClearColor, const FLinearColor& Color, bool bClearDepth, float Depth, bool bClearStencil, uint32 Stencil)
 {
-	DrawClearQuadMRT(RHICmdList, bClearColor, 1, &Color, bClearDepth, Depth, bClearStencil, Stencil);
+	DrawClearQuadMRT(RHICmdList, FeatureLevel, bClearColor, 1, &Color, bClearDepth, Depth, bClearStencil, Stencil);
 }
\ No newline at end of file
diff --git a/Engine/Source/Runtime/Renderer/Private/CompositionLighting/PostProcessAmbient.cpp b/Engine/Source/Runtime/Renderer/Private/CompositionLighting/PostProcessAmbient.cpp
--- a/Engine/Source/Runtime/Renderer/Private/CompositionLighting/PostProcessAmbient.cpp
+++ b/Engine/Source/Runtime/Renderer/Private/CompositionLighting/PostProcessAmbient.cpp
@@ -110,7 +110,7 @@ void FRCPassPostProcessAmbient::Process(FRenderingCompositePassContext& Context)
 		if(i == 0)
 		{
 			// call it once after setting up the shader data to avoid the warnings in the function
-			SetGlobalBoundShaderState(Context.RHICmdList, BoundShaderState, GFilterVertexDeclaration.VertexDeclarationRHI, *VertexShader, *PixelShader);
+			SetGlobalBoundShaderState(Context.RHICmdList, Context.GetFeatureLevel(), BoundShaderState, GFilterVertexDeclaration.VertexDeclarationRHI, *VertexShader, *PixelShader);
 		}
 
 		
diff --git a/Engine/Source/Runtime/Renderer/Private/CompositionLighting/PostProcessAmbientOcclusion.cpp b/Engine/Source/Runtime/Renderer/Private/CompositionLighting/PostProcessAmbientOcclusion.cpp
--- a/Engine/Source/Runtime/Renderer/Private/CompositionLighting/PostProcessAmbientOcclusion.cpp
+++ b/Engine/Source/Runtime/Renderer/Private/CompositionLighting/PostProcessAmbientOcclusion.cpp
@@ -352,7 +352,7 @@ void FRCPassPostProcessAmbientOcclusionSetup::SetShaderSetupTempl(const FRenderi
 	static FGlobalBoundShaderState BoundShaderState;
 	
 
-	SetGlobalBoundShaderState(Context.RHICmdList, BoundShaderState, GFilterVertexDeclaration.VertexDeclarationRHI, *VertexShader, *PixelShader);
+	SetGlobalBoundShaderState(Context.RHICmdList, Context.GetFeatureLevel(), BoundShaderState, GFilterVertexDeclaration.VertexDeclarationRHI, *VertexShader, *PixelShader);
 
 	VertexShader->SetParameters(Context);
 	PixelShader->SetParameters(Context);
@@ -463,7 +463,7 @@ void FRCPassPostProcessAmbientOcclusion::SetShaderTempl(const FRenderingComposit
 	static FGlobalBoundShaderState BoundShaderState;
 	
 
-	SetGlobalBoundShaderState(Context.RHICmdList, BoundShaderState, GFilterVertexDeclaration.VertexDeclarationRHI, *VertexShader, *PixelShader);
+	SetGlobalBoundShaderState(Context.RHICmdList, Context.GetFeatureLevel(), BoundShaderState, GFilterVertexDeclaration.VertexDeclarationRHI, *VertexShader, *PixelShader);
 
 	const FPooledRenderTargetDesc* InputDesc0 = GetInputDesc(ePId_Input0);
 
@@ -611,7 +611,7 @@ void FRCPassPostProcessBasePassAO::Process(FRenderingCompositePassContext& Conte
 	static FGlobalBoundShaderState BoundShaderState;
 	
 
-	SetGlobalBoundShaderState(Context.RHICmdList, BoundShaderState, GFilterVertexDeclaration.VertexDeclarationRHI, *VertexShader, *PixelShader);
+	SetGlobalBoundShaderState(Context.RHICmdList, Context.GetFeatureLevel(), BoundShaderState, GFilterVertexDeclaration.VertexDeclarationRHI, *VertexShader, *PixelShader);
 
 	VertexShader->SetParameters(Context);
 	PixelShader->SetParameters(Context, GSceneRenderTargets.GetBufferSizeXY());
diff --git a/Engine/Source/Runtime/Renderer/Private/CompositionLighting/PostProcessDeferredDecals.cpp b/Engine/Source/Runtime/Renderer/Private/CompositionLighting/PostProcessDeferredDecals.cpp
--- a/Engine/Source/Runtime/Renderer/Private/CompositionLighting/PostProcessDeferredDecals.cpp
+++ b/Engine/Source/Runtime/Renderer/Private/CompositionLighting/PostProcessDeferredDecals.cpp
@@ -364,7 +364,7 @@ void StencilDecalMask(FRHICommandList& RHICmdList, const FSceneView& View)
 	TShaderMapRef<FScreenVS> ScreenVertexShader(GetGlobalShaderMap());
 	TShaderMapRef<FStencilDecalMaskPS> PixelShader(GetGlobalShaderMap());
 	
-	SetGlobalBoundShaderState(RHICmdList, StencilDecalMaskBoundShaderState, GFilterVertexDeclaration.VertexDeclarationRHI, *ScreenVertexShader, *PixelShader);
+	SetGlobalBoundShaderState(RHICmdList, View.GetFeatureLevel(), StencilDecalMaskBoundShaderState, GFilterVertexDeclaration.VertexDeclarationRHI, *ScreenVertexShader, *PixelShader);
 
 	PixelShader->SetParameters(RHICmdList, View);
 
diff --git a/Engine/Source/Runtime/Renderer/Private/CompositionLighting/PostProcessLpvIndirect.cpp b/Engine/Source/Runtime/Renderer/Private/CompositionLighting/PostProcessLpvIndirect.cpp
--- a/Engine/Source/Runtime/Renderer/Private/CompositionLighting/PostProcessLpvIndirect.cpp
+++ b/Engine/Source/Runtime/Renderer/Private/CompositionLighting/PostProcessLpvIndirect.cpp
@@ -191,7 +191,7 @@ void FRCPassPostProcessLpvIndirect::Process(FRenderingCompositePassContext& Cont
 	
 
 	// call it once after setting up the shader data to avoid the warnings in the function
-	SetGlobalBoundShaderState(Context.RHICmdList, BoundShaderState, GFilterVertexDeclaration.VertexDeclarationRHI, *VertexShader, *PixelShader);
+	SetGlobalBoundShaderState(Context.RHICmdList, Context.GetFeatureLevel(), BoundShaderState, GFilterVertexDeclaration.VertexDeclarationRHI, *VertexShader, *PixelShader);
 
 	FLpvReadUniformBufferParameters	LpvReadUniformBufferParams;
 	FLpvReadUniformBufferRef LpvReadUniformBuffer;
diff --git a/Engine/Source/Runtime/Renderer/Private/CompositionLighting/PostProcessPassThrough.cpp b/Engine/Source/Runtime/Renderer/Private/CompositionLighting/PostProcessPassThrough.cpp
--- a/Engine/Source/Runtime/Renderer/Private/CompositionLighting/PostProcessPassThrough.cpp
+++ b/Engine/Source/Runtime/Renderer/Private/CompositionLighting/PostProcessPassThrough.cpp
@@ -120,7 +120,7 @@ void FRCPassPostProcessPassThrough::Process(FRenderingCompositePassContext& Cont
 	static FGlobalBoundShaderState BoundShaderState;
 	
 
-	SetGlobalBoundShaderState(Context.RHICmdList, BoundShaderState, GFilterVertexDeclaration.VertexDeclarationRHI, *VertexShader, *PixelShader);
+	SetGlobalBoundShaderState(Context.RHICmdList, Context.GetFeatureLevel(), BoundShaderState, GFilterVertexDeclaration.VertexDeclarationRHI, *VertexShader, *PixelShader);
 
 	VertexShader->SetParameters(Context);
 	PixelShader->SetParameters(Context);
diff --git a/Engine/Source/Runtime/Renderer/Private/DeferredShadingRenderer.cpp b/Engine/Source/Runtime/Renderer/Private/DeferredShadingRenderer.cpp
--- a/Engine/Source/Runtime/Renderer/Private/DeferredShadingRenderer.cpp
+++ b/Engine/Source/Runtime/Renderer/Private/DeferredShadingRenderer.cpp
@@ -171,7 +171,7 @@ void FDeferredShadingSceneRenderer::ClearGBufferAtMaxZ(FRHICommandList& RHICmdLi
 		break;
 	}
 
-	SetGlobalBoundShaderState(RHICmdList, GClearMRTBoundShaderState[NumActiveRenderTargets - 1], GetVertexDeclarationFVector4(), *VertexShader, PixelShader);
+	SetGlobalBoundShaderState(RHICmdList, FeatureLevel, GClearMRTBoundShaderState[NumActiveRenderTargets - 1], GetVertexDeclarationFVector4(), *VertexShader, PixelShader);
 
 	// Opaque rendering, depth test but no depth writes
 	RHICmdList.SetRasterizerState(TStaticRasterizerState<FM_Solid, CM_None>::GetRHI());
@@ -411,20 +411,20 @@ void FDeferredShadingSceneRenderer::RenderBasePassDynamicData(FRHICommandList& R
 
 	if( !View.Family->EngineShowFlags.CompositeEditorPrimitives )
 	{
-		
-		const bool bNeedToSwitchVerticalAxis = RHINeedsToSwitchVerticalAxis(GRHIShaderPlatform);
+		const auto ShaderPlatform = View.GetShaderPlatform();
+		const bool bNeedToSwitchVerticalAxis = RHINeedsToSwitchVerticalAxis(ShaderPlatform);
 
 		// Draw the base pass for the view's batched mesh elements.
 		bDirty = DrawViewElements<FBasePassOpaqueDrawingPolicyFactory>(RHICmdList, View, FBasePassOpaqueDrawingPolicyFactory::ContextType(false, ESceneRenderTargetsMode::DontSet), SDPG_World, true) || bDirty;
 
 		// Draw the view's batched simple elements(lines, sprites, etc).
-		bDirty = View.BatchedViewElements.Draw(RHICmdList, bNeedToSwitchVerticalAxis, View.ViewProjectionMatrix, View.ViewRect.Width(), View.ViewRect.Height(), false) || bDirty;
+		bDirty = View.BatchedViewElements.Draw(RHICmdList, FeatureLevel, bNeedToSwitchVerticalAxis, View.ViewProjectionMatrix, View.ViewRect.Width(), View.ViewRect.Height(), false) || bDirty;
 		
 		// Draw foreground objects last
 		bDirty = DrawViewElements<FBasePassOpaqueDrawingPolicyFactory>(RHICmdList, View, FBasePassOpaqueDrawingPolicyFactory::ContextType(false, ESceneRenderTargetsMode::DontSet), SDPG_Foreground, true) || bDirty;
 
 		// Draw the view's batched simple elements(lines, sprites, etc).
-		bDirty = View.TopBatchedViewElements.Draw(RHICmdList, bNeedToSwitchVerticalAxis, View.ViewProjectionMatrix, View.ViewRect.Width(), View.ViewRect.Height(), false) || bDirty;
+		bDirty = View.TopBatchedViewElements.Draw(RHICmdList, FeatureLevel, bNeedToSwitchVerticalAxis, View.ViewProjectionMatrix, View.ViewRect.Width(), View.ViewRect.Height(), false) || bDirty;
 
 	}
 
@@ -667,7 +667,7 @@ void FDeferredShadingSceneRenderer::Render(FRHICommandListImmediate& RHICmdList)
 	#endif
 
 	
-	const bool bIsOcclusionTesting = DoOcclusionQueries() && (!bIsWireframe || bIsViewFrozen || bHasViewParent);
+	const bool bIsOcclusionTesting = DoOcclusionQueries(FeatureLevel) && (!bIsWireframe || bIsViewFrozen || bHasViewParent);
 
 	// Dynamic vertex and index buffers need to be committed before rendering.
 	FGlobalDynamicVertexBuffer::Get().Commit();
@@ -696,7 +696,7 @@ void FDeferredShadingSceneRenderer::Render(FRHICommandListImmediate& RHICmdList)
 	}
 
 	// Clear LPVs for all views
-	if ( IsFeatureLevelSupported(GRHIShaderPlatform, ERHIFeatureLevel::SM5) )
+	if ( FeatureLevel >= ERHIFeatureLevel::SM5 )
 	{
 		ClearLPVs(RHICmdList);
 	}
@@ -820,7 +820,7 @@ void FDeferredShadingSceneRenderer::Render(FRHICommandListImmediate& RHICmdList)
 		FilterTranslucentVolumeLighting(RHICmdList);
 
 		// Clear LPVs for all views
-		if ( IsFeatureLevelSupported(GRHIShaderPlatform, ERHIFeatureLevel::SM5) )
+		if ( FeatureLevel >= ERHIFeatureLevel::SM5 )
 		{
 			PropagateLPVs(RHICmdList);
 		}
@@ -1363,7 +1363,7 @@ void FDeferredShadingSceneRenderer::UpdateDownsampledDepthSurface(FRHICommandLis
 
 			extern TGlobalResource<FFilterVertexDeclaration> GFilterVertexDeclaration;
 
-			SetGlobalBoundShaderState(RHICmdList, DownsampleDepthBoundShaderState, GFilterVertexDeclaration.VertexDeclarationRHI, *ScreenVertexShader, *PixelShader);
+			SetGlobalBoundShaderState(RHICmdList, FeatureLevel, DownsampleDepthBoundShaderState, GFilterVertexDeclaration.VertexDeclarationRHI, *ScreenVertexShader, *PixelShader);
 
 			RHICmdList.SetBlendState(TStaticBlendState<CW_NONE>::GetRHI());
 			RHICmdList.SetRasterizerState(TStaticRasterizerState<FM_Solid, CM_None>::GetRHI());
diff --git a/Engine/Source/Runtime/Renderer/Private/DistanceFieldSurfaceCacheLighting.cpp b/Engine/Source/Runtime/Renderer/Private/DistanceFieldSurfaceCacheLighting.cpp
--- a/Engine/Source/Runtime/Renderer/Private/DistanceFieldSurfaceCacheLighting.cpp
+++ b/Engine/Source/Runtime/Renderer/Private/DistanceFieldSurfaceCacheLighting.cpp
@@ -986,7 +986,7 @@ void ComputeDistanceFieldNormal(FRHICommandListImmediate& RHICmdList, const TArr
 
 			static FGlobalBoundShaderState BoundShaderState;
 			
-			SetGlobalBoundShaderState(RHICmdList, BoundShaderState, GFilterVertexDeclaration.VertexDeclarationRHI, *VertexShader, *PixelShader);
+			SetGlobalBoundShaderState(RHICmdList, View.GetFeatureLevel(), BoundShaderState, GFilterVertexDeclaration.VertexDeclarationRHI, *VertexShader, *PixelShader);
 
 			PixelShader->SetParameters(RHICmdList, View, NumObjects);
 
@@ -2086,7 +2086,7 @@ void UpdateHistory(
 
 				static FGlobalBoundShaderState BoundShaderState;
 				
-				SetGlobalBoundShaderState(RHICmdList, BoundShaderState, GFilterVertexDeclaration.VertexDeclarationRHI, *VertexShader, *PixelShader);
+				SetGlobalBoundShaderState(RHICmdList, View.GetFeatureLevel(), BoundShaderState, GFilterVertexDeclaration.VertexDeclarationRHI, *VertexShader, *PixelShader);
 
 				PixelShader->SetParameters(RHICmdList, View, (*HistoryState)->GetRenderTargetItem(), AOSource->GetRenderTargetItem());
 
@@ -2154,7 +2154,7 @@ void PostProcessBentNormalAO(FRHICommandList& RHICmdList, TArray<FViewInfo>& Vie
 
 			static FGlobalBoundShaderState BoundShaderState;
 			
-			SetGlobalBoundShaderState(RHICmdList, BoundShaderState, GFilterVertexDeclaration.VertexDeclarationRHI, *VertexShader, *PixelShader);
+			SetGlobalBoundShaderState(RHICmdList, View.GetFeatureLevel(), BoundShaderState, GFilterVertexDeclaration.VertexDeclarationRHI, *VertexShader, *PixelShader);
 
 			PixelShader->SetParameters(RHICmdList, View, IrradianceCacheInterpolation);
 
@@ -2189,7 +2189,7 @@ void PostProcessBentNormalAO(FRHICommandList& RHICmdList, TArray<FViewInfo>& Vie
 
 			static FGlobalBoundShaderState BoundShaderState;
 			
-			SetGlobalBoundShaderState(RHICmdList, BoundShaderState, GFilterVertexDeclaration.VertexDeclarationRHI, *VertexShader, *PixelShader);
+			SetGlobalBoundShaderState(RHICmdList, View.GetFeatureLevel(), BoundShaderState, GFilterVertexDeclaration.VertexDeclarationRHI, *VertexShader, *PixelShader);
 
 			VertexShader->SetParameters(RHICmdList, View);
 			PixelShader->SetParameters(RHICmdList, View, DistanceFieldAOBentNormal2->GetRenderTargetItem());
@@ -2297,7 +2297,7 @@ void UpsampleBentNormalAO(FRHICommandList& RHICmdList, const TArray<FViewInfo>&
 
 		static FGlobalBoundShaderState BoundShaderState;
 		
-		SetGlobalBoundShaderState(RHICmdList, BoundShaderState, GFilterVertexDeclaration.VertexDeclarationRHI, *VertexShader, *PixelShader);
+		SetGlobalBoundShaderState(RHICmdList, View.GetFeatureLevel(), BoundShaderState, GFilterVertexDeclaration.VertexDeclarationRHI, *VertexShader, *PixelShader);
 
 		PixelShader->SetParameters(RHICmdList, View, DistanceFieldAOBentNormal);
 
@@ -2536,7 +2536,7 @@ FIntPoint BuildTileObjectLists(FRHICommandListImmediate& RHICmdList, FScene* Sce
 
 				static FGlobalBoundShaderState BoundShaderState;
 				
-				SetGlobalBoundShaderState(RHICmdList, BoundShaderState, GetVertexDeclarationFVector4(), *VertexShader, *PixelShader);
+				SetGlobalBoundShaderState(RHICmdList, View.GetFeatureLevel(), BoundShaderState, GetVertexDeclarationFVector4(), *VertexShader, *PixelShader);
 
 				VertexShader->SetParameters(RHICmdList, View, NumObjects);
 				PixelShader->SetParameters(RHICmdList, View, FVector2D(GroupSizeX, GroupSizeY));
@@ -2660,7 +2660,7 @@ void RenderIrradianceCacheInterpolation(FRHICommandListImmediate& RHICmdList, TA
 				{
 					static FGlobalBoundShaderState BoundShaderState;
 					
-					SetGlobalBoundShaderState(RHICmdList, BoundShaderState, GScreenVertexDeclaration.VertexDeclarationRHI, *VertexShader, *PixelShader);
+					SetGlobalBoundShaderState(RHICmdList, View.GetFeatureLevel(), BoundShaderState, GScreenVertexDeclaration.VertexDeclarationRHI, *VertexShader, *PixelShader);
 
 					VertexShader->SetParameters(RHICmdList, View, SplatSourceDepthLevel, DestLevelDownsampleFactor, NormalizedOffsetToPixelCenter);
 					PixelShader->SetParameters(RHICmdList, View, DistanceFieldNormal, DestLevelDownsampleFactor);
@@ -2677,7 +2677,7 @@ void RenderIrradianceCacheInterpolation(FRHICommandListImmediate& RHICmdList, TA
 				{
 					static FGlobalBoundShaderState BoundShaderState;
 					
-					SetGlobalBoundShaderState(RHICmdList, BoundShaderState, GScreenVertexDeclaration.VertexDeclarationRHI, *VertexShader, *PixelShader);
+					SetGlobalBoundShaderState(RHICmdList, View.GetFeatureLevel(), BoundShaderState, GScreenVertexDeclaration.VertexDeclarationRHI, *VertexShader, *PixelShader);
 
 					VertexShader->SetParameters(RHICmdList, View, SplatSourceDepthLevel, DestLevelDownsampleFactor, NormalizedOffsetToPixelCenter);
 					PixelShader->SetParameters(RHICmdList, View, DistanceFieldNormal, DestLevelDownsampleFactor);
@@ -3146,7 +3146,7 @@ void FDeferredShadingSceneRenderer::RenderDynamicSkyLighting(FRHICommandListImme
 
 				static FGlobalBoundShaderState BoundShaderState;
 				
-				SetGlobalBoundShaderState(RHICmdList, BoundShaderState, GFilterVertexDeclaration.VertexDeclarationRHI, *VertexShader, *PixelShader);
+				SetGlobalBoundShaderState(RHICmdList, View.GetFeatureLevel(), BoundShaderState, GFilterVertexDeclaration.VertexDeclarationRHI, *VertexShader, *PixelShader);
 
 				PixelShader->SetParameters(RHICmdList, View, LightAccumulation->GetRenderTargetItem().ShaderResourceTexture);
 			}
@@ -3156,7 +3156,7 @@ void FDeferredShadingSceneRenderer::RenderDynamicSkyLighting(FRHICommandListImme
 
 				static FGlobalBoundShaderState BoundShaderState;
 				
-				SetGlobalBoundShaderState(RHICmdList, BoundShaderState, GFilterVertexDeclaration.VertexDeclarationRHI, *VertexShader, *PixelShader);
+				SetGlobalBoundShaderState(RHICmdList, View.GetFeatureLevel(), BoundShaderState, GFilterVertexDeclaration.VertexDeclarationRHI, *VertexShader, *PixelShader);
 
 				PixelShader->SetParameters(RHICmdList, View, LightAccumulation->GetRenderTargetItem().ShaderResourceTexture);
 			}
diff --git a/Engine/Source/Runtime/Renderer/Private/DistortionRendering.cpp b/Engine/Source/Runtime/Renderer/Private/DistortionRendering.cpp
--- a/Engine/Source/Runtime/Renderer/Private/DistortionRendering.cpp
+++ b/Engine/Source/Runtime/Renderer/Private/DistortionRendering.cpp
@@ -899,7 +899,7 @@ void FSceneRenderer::RenderDistortion(FRHICommandListImmediate& RHICmdList)
 
 				static FGlobalBoundShaderState BoundShaderState;
 				
-				SetGlobalBoundShaderState(RHICmdList, BoundShaderState, GFilterVertexDeclaration.VertexDeclarationRHI, *VertexShader, *PixelShader);
+				SetGlobalBoundShaderState(RHICmdList, View.GetFeatureLevel(), BoundShaderState, GFilterVertexDeclaration.VertexDeclarationRHI, *VertexShader, *PixelShader);
 
 				VertexShader->SetParameters(Context);
 				PixelShader->SetParameters(Context, *DistortionRT);
diff --git a/Engine/Source/Runtime/Renderer/Private/DynamicPrimitiveDrawing.inl b/Engine/Source/Runtime/Renderer/Private/DynamicPrimitiveDrawing.inl
--- a/Engine/Source/Runtime/Renderer/Private/DynamicPrimitiveDrawing.inl
+++ b/Engine/Source/Runtime/Renderer/Private/DynamicPrimitiveDrawing.inl
@@ -12,7 +12,8 @@ TDynamicPrimitiveDrawer<DrawingPolicyFactoryType>::~TDynamicPrimitiveDrawer()
 {
 	if(View)
 	{
-		const bool bNeedToSwitchVerticalAxis = RHINeedsToSwitchVerticalAxis(GRHIShaderPlatform) && !IsMobileHDR();
+		const auto ShaderPlatform = View->GetShaderPlatform();
+		const bool bNeedToSwitchVerticalAxis = RHINeedsToSwitchVerticalAxis(ShaderPlatform) && !IsMobileHDR();
 
 		// Determine whether or not to depth test simple batched elements manually (msaa only)
 		// We need to read from the scene depth in the pixel shader to do manual depth testing
@@ -26,6 +27,7 @@ TDynamicPrimitiveDrawer<DrawingPolicyFactoryType>::~TDynamicPrimitiveDrawer()
 		// Draw the batched elements.
 		BatchedElements.Draw(
 			RHICmdList,
+			View->GetFeatureLevel(),
 			bNeedToSwitchVerticalAxis,
 			View->ViewProjectionMatrix,
 			View->ViewRect.Width(),
diff --git a/Engine/Source/Runtime/Renderer/Private/FogRendering.cpp b/Engine/Source/Runtime/Renderer/Private/FogRendering.cpp
--- a/Engine/Source/Runtime/Renderer/Private/FogRendering.cpp
+++ b/Engine/Source/Runtime/Renderer/Private/FogRendering.cpp
@@ -282,7 +282,7 @@ void SetFogShaders(FRHICommandList& RHICmdList, FScene* Scene, const FViewInfo&
 		TShaderMapRef<FHeightFogVS> VertexShader(GetGlobalShaderMap());
 		TShaderMapRef<FExponentialHeightFogPS> ExponentialHeightFogPixelShader(GetGlobalShaderMap());
 
-		SetGlobalBoundShaderState(RHICmdList, ExponentialBoundShaderState, GFogVertexDeclaration.VertexDeclarationRHI, *VertexShader, *ExponentialHeightFogPixelShader);
+		SetGlobalBoundShaderState(RHICmdList, View.GetFeatureLevel(), ExponentialBoundShaderState, GFogVertexDeclaration.VertexDeclarationRHI, *VertexShader, *ExponentialHeightFogPixelShader);
 		VertexShader->SetParameters(RHICmdList, View);
 		ExponentialHeightFogPixelShader->SetParameters(RHICmdList, View, LightShaftsOutput);
 	}
diff --git a/Engine/Source/Runtime/Renderer/Private/ForwardBasePassRendering.cpp b/Engine/Source/Runtime/Renderer/Private/ForwardBasePassRendering.cpp
--- a/Engine/Source/Runtime/Renderer/Private/ForwardBasePassRendering.cpp
+++ b/Engine/Source/Runtime/Renderer/Private/ForwardBasePassRendering.cpp
@@ -215,6 +215,7 @@ bool FBasePassForwardOpaqueDrawingPolicyFactory::DrawDynamicMesh(
 {
 	// Determine the mesh's material and blend mode.
 	const auto FeatureLevel = View.GetFeatureLevel();
+	const auto ShaderPlatform = View.GetShaderPlatform();
 	const FMaterial* Material = Mesh.MaterialRenderProxy->GetMaterial(FeatureLevel);
 	const EBlendMode BlendMode = Material->GetBlendMode();
 
@@ -403,13 +404,13 @@ void FForwardShadingSceneRenderer::RenderForwardShadingBasePass(FRHICommandListI
 			DrawViewElements<FBasePassForwardOpaqueDrawingPolicyFactory>(RHICmdList, View, FBasePassForwardOpaqueDrawingPolicyFactory::ContextType(ESceneRenderTargetsMode::DontSet), SDPG_World, true);
 			
 			// Draw the view's batched simple elements(lines, sprites, etc).
-			View.BatchedViewElements.Draw(RHICmdList, bNeedToSwitchVerticalAxis, View.ViewProjectionMatrix, View.ViewRect.Width(), View.ViewRect.Height(), false);
+			View.BatchedViewElements.Draw(RHICmdList, FeatureLevel, bNeedToSwitchVerticalAxis, View.ViewProjectionMatrix, View.ViewRect.Width(), View.ViewRect.Height(), false);
 
 			// Draw foreground objects last
 			DrawViewElements<FBasePassForwardOpaqueDrawingPolicyFactory>(RHICmdList, View, FBasePassForwardOpaqueDrawingPolicyFactory::ContextType(ESceneRenderTargetsMode::DontSet), SDPG_Foreground, true);
 
 			// Draw the view's batched simple elements(lines, sprites, etc).
-			View.TopBatchedViewElements.Draw(RHICmdList, bNeedToSwitchVerticalAxis, View.ViewProjectionMatrix, View.ViewRect.Width(), View.ViewRect.Height(), false);
+			View.TopBatchedViewElements.Draw(RHICmdList, FeatureLevel, bNeedToSwitchVerticalAxis, View.ViewProjectionMatrix, View.ViewRect.Width(), View.ViewRect.Height(), false);
 		}
 
 		// Issue static draw list masked draw calls last, as PVR wants it
diff --git a/Engine/Source/Runtime/Renderer/Private/ForwardTranslucentRendering.cpp b/Engine/Source/Runtime/Renderer/Private/ForwardTranslucentRendering.cpp
--- a/Engine/Source/Runtime/Renderer/Private/ForwardTranslucentRendering.cpp
+++ b/Engine/Source/Runtime/Renderer/Private/ForwardTranslucentRendering.cpp
@@ -63,7 +63,7 @@ void FForwardShadingSceneRenderer::CopySceneAlpha(FRHICommandListImmediate& RHIC
 
 	TShaderMapRef<FScreenVS> ScreenVertexShader(GetGlobalShaderMap());
 	TShaderMapRef<FForwardCopySceneAlphaPS> PixelShader(GetGlobalShaderMap());
-	SetGlobalBoundShaderState(RHICmdList, ForwardCopySceneAlphaBoundShaderState, GFilterVertexDeclaration.VertexDeclarationRHI, *ScreenVertexShader, *PixelShader);
+	SetGlobalBoundShaderState(RHICmdList, FeatureLevel, ForwardCopySceneAlphaBoundShaderState, GFilterVertexDeclaration.VertexDeclarationRHI, *ScreenVertexShader, *PixelShader);
 
 	PixelShader->SetParameters(RHICmdList, View);
 
@@ -177,6 +177,7 @@ bool FTranslucencyForwardShadingDrawingPolicyFactory::DrawDynamicMesh(
 
 	// Determine the mesh's material and blend mode.
 	const auto FeatureLevel = View.GetFeatureLevel();
+	const auto ShaderPlatform = View.GetShaderPlatform();
 	const FMaterial* Material = Mesh.MaterialRenderProxy->GetMaterial(FeatureLevel);
 	const EBlendMode BlendMode = Material->GetBlendMode();
 
diff --git a/Engine/Source/Runtime/Renderer/Private/GPUBenchmark.cpp b/Engine/Source/Runtime/Renderer/Private/GPUBenchmark.cpp
--- a/Engine/Source/Runtime/Renderer/Private/GPUBenchmark.cpp
+++ b/Engine/Source/Runtime/Renderer/Private/GPUBenchmark.cpp
@@ -133,7 +133,7 @@ void RunBenchmarkShader(FRHICommandList& RHICmdList, const FSceneView& View, TRe
 	static FGlobalBoundShaderState BoundShaderState;
 	
 
-	SetGlobalBoundShaderState(RHICmdList, BoundShaderState, GFilterVertexDeclaration.VertexDeclarationRHI, *VertexShader, *PixelShader);
+	SetGlobalBoundShaderState(RHICmdList, View.GetFeatureLevel(), BoundShaderState, GFilterVertexDeclaration.VertexDeclarationRHI, *VertexShader, *PixelShader);
 
 	PixelShader->SetParameters(RHICmdList, View, Src);
 	VertexShader->SetParameters(RHICmdList, View);
diff --git a/Engine/Source/Runtime/Renderer/Private/LightPropagationVolume.cpp b/Engine/Source/Runtime/Renderer/Private/LightPropagationVolume.cpp
--- a/Engine/Source/Runtime/Renderer/Private/LightPropagationVolume.cpp
+++ b/Engine/Source/Runtime/Renderer/Private/LightPropagationVolume.cpp
@@ -1250,9 +1250,9 @@ void FLightPropagationVolume::InjectLightDirect(FRHICommandListImmediate& RHICmd
 
 
 // use for render thread only
-bool UseLightPropagationVolumeRT()
+bool UseLightPropagationVolumeRT(ERHIFeatureLevel::Type InFeatureLevel)
 {
-	if(!IsFeatureLevelSupported(GRHIShaderPlatform, ERHIFeatureLevel::SM5))
+	if (InFeatureLevel < ERHIFeatureLevel::SM5)
 	{
 		return false;
 	}
@@ -1266,7 +1266,7 @@ bool UseLightPropagationVolumeRT()
 // ----------------------------------------------------------------------------
 // FSceneViewState
 // ----------------------------------------------------------------------------
-void FSceneViewState::CreateLightPropagationVolumeIfNeeded()
+void FSceneViewState::CreateLightPropagationVolumeIfNeeded(ERHIFeatureLevel::Type InFeatureLevel)
 {
 	check(IsInRenderingThread());
 
@@ -1276,7 +1276,7 @@ void FSceneViewState::CreateLightPropagationVolumeIfNeeded()
 		return;
 	}
 
-	bool bLPV = UseLightPropagationVolumeRT();
+	bool bLPV = UseLightPropagationVolumeRT(InFeatureLevel);
 
 	//@todo-rco: Remove this when reenabling for OpenGL
 	if (bLPV && !IsOpenGLPlatform(GRHIShaderPlatform))
diff --git a/Engine/Source/Runtime/Renderer/Private/LightPropagationVolume.h b/Engine/Source/Runtime/Renderer/Private/LightPropagationVolume.h
--- a/Engine/Source/Runtime/Renderer/Private/LightPropagationVolume.h
+++ b/Engine/Source/Runtime/Renderer/Private/LightPropagationVolume.h
@@ -225,4 +225,4 @@ public:
 
 
 // use for render thread only
-bool UseLightPropagationVolumeRT();
+bool UseLightPropagationVolumeRT(ERHIFeatureLevel::Type InFeatureLevel);
diff --git a/Engine/Source/Runtime/Renderer/Private/LightPropagationVolumeVisualisation.cpp b/Engine/Source/Runtime/Renderer/Private/LightPropagationVolumeVisualisation.cpp
--- a/Engine/Source/Runtime/Renderer/Private/LightPropagationVolumeVisualisation.cpp
+++ b/Engine/Source/Runtime/Renderer/Private/LightPropagationVolumeVisualisation.cpp
@@ -260,7 +260,7 @@ void FLightPropagationVolume::Visualise(FRHICommandList& RHICmdList, const FScen
 	RHICmdList.SetRasterizerState(TStaticRasterizerState<FM_Solid, CM_None>::GetRHI());
 	RHICmdList.SetBlendState(TStaticBlendState<CW_RGB, BO_Add, BF_One, BF_One>::GetRHI());
 
-	SetGlobalBoundShaderState(RHICmdList, LpvVisBoundShaderState, GSimpleElementVertexDeclaration.VertexDeclarationRHI, *VertexShader, *PixelShader, *GeometryShader);
+	SetGlobalBoundShaderState(RHICmdList, View.GetFeatureLevel(), LpvVisBoundShaderState, GSimpleElementVertexDeclaration.VertexDeclarationRHI, *VertexShader, *PixelShader, *GeometryShader);
 
 	VertexShader->SetParameters(RHICmdList, View);
 	GeometryShader->SetParameters(RHICmdList, View);
diff --git a/Engine/Source/Runtime/Renderer/Private/LightRendering.cpp b/Engine/Source/Runtime/Renderer/Private/LightRendering.cpp
--- a/Engine/Source/Runtime/Renderer/Private/LightRendering.cpp
+++ b/Engine/Source/Runtime/Renderer/Private/LightRendering.cpp
@@ -672,13 +672,13 @@ static void SetShaderTemplLighting(
 	if(View.Family->EngineShowFlags.VisualizeLightCulling)
 	{
 		TShaderMapRef<TDeferredLightPS<false, bRadialAttenuation, false, true> > PixelShader(GetGlobalShaderMap());
-		SetGlobalBoundShaderState(RHICmdList, PixelShader->GetBoundShaderState(), GetDeferredLightingVertexDeclaration<bRadialAttenuation>(), VertexShader, *PixelShader);
+		SetGlobalBoundShaderState(RHICmdList, View.GetFeatureLevel(), PixelShader->GetBoundShaderState(), GetDeferredLightingVertexDeclaration<bRadialAttenuation>(), VertexShader, *PixelShader);
 		PixelShader->SetParameters(RHICmdList, View, LightSceneInfo);
 	}
 	else
 	{
 		TShaderMapRef<TDeferredLightPS<bUseIESProfile, bRadialAttenuation, bInverseSquaredFalloff, false> > PixelShader(GetGlobalShaderMap());
-		SetGlobalBoundShaderState(RHICmdList, PixelShader->GetBoundShaderState(), GetDeferredLightingVertexDeclaration<bRadialAttenuation>(), VertexShader, *PixelShader);
+		SetGlobalBoundShaderState(RHICmdList, View.GetFeatureLevel(), PixelShader->GetBoundShaderState(), GetDeferredLightingVertexDeclaration<bRadialAttenuation>(), VertexShader, *PixelShader);
 		PixelShader->SetParameters(RHICmdList, View, LightSceneInfo);
 	}
 }
@@ -694,13 +694,13 @@ static void SetShaderTemplLightingSimple(
 	if(View.Family->EngineShowFlags.VisualizeLightCulling)
 	{
 		TShaderMapRef<TDeferredLightPS<false, bRadialAttenuation, false, true> > PixelShader(GetGlobalShaderMap());
-		SetGlobalBoundShaderState(RHICmdList, PixelShader->GetBoundShaderState(), GetDeferredLightingVertexDeclaration<bRadialAttenuation>(), VertexShader, *PixelShader);
+		SetGlobalBoundShaderState(RHICmdList, View.GetFeatureLevel(), PixelShader->GetBoundShaderState(), GetDeferredLightingVertexDeclaration<bRadialAttenuation>(), VertexShader, *PixelShader);
 		PixelShader->SetParametersSimpleLight(RHICmdList, View, SimpleLight, SimpleLightPerViewData);
 	}
 	else
 	{
 		TShaderMapRef<TDeferredLightPS<bUseIESProfile, bRadialAttenuation, bInverseSquaredFalloff, false> > PixelShader(GetGlobalShaderMap());
-		SetGlobalBoundShaderState(RHICmdList, PixelShader->GetBoundShaderState(), GetDeferredLightingVertexDeclaration<bRadialAttenuation>(), VertexShader, *PixelShader);
+		SetGlobalBoundShaderState(RHICmdList, View.GetFeatureLevel(), PixelShader->GetBoundShaderState(), GetDeferredLightingVertexDeclaration<bRadialAttenuation>(), VertexShader, *PixelShader);
 		PixelShader->SetParametersSimpleLight(RHICmdList, View, SimpleLight, SimpleLightPerViewData);
 	}
 }
@@ -748,7 +748,7 @@ void FDeferredShadingSceneRenderer::RenderLight(FRHICommandList& RHICmdList, con
 			if (bRenderOverlap)
 			{
 				TShaderMapRef<TDeferredLightOverlapPS<false> > PixelShader(GetGlobalShaderMap());
-				SetGlobalBoundShaderState(RHICmdList, PixelShader->GetBoundShaderState(), GetDeferredLightingVertexDeclaration<false>(), *VertexShader, *PixelShader);
+				SetGlobalBoundShaderState(RHICmdList, FeatureLevel, PixelShader->GetBoundShaderState(), GetDeferredLightingVertexDeclaration<false>(), *VertexShader, *PixelShader);
 				PixelShader->SetParameters(RHICmdList, View, LightSceneInfo);
 			}
 			else
@@ -786,7 +786,7 @@ void FDeferredShadingSceneRenderer::RenderLight(FRHICommandList& RHICmdList, con
 			if (bRenderOverlap)
 			{
 				TShaderMapRef<TDeferredLightOverlapPS<true> > PixelShader(GetGlobalShaderMap());
-				SetGlobalBoundShaderState(RHICmdList, PixelShader->GetBoundShaderState(), GetDeferredLightingVertexDeclaration<true>(), *VertexShader, *PixelShader);
+				SetGlobalBoundShaderState(RHICmdList, FeatureLevel, PixelShader->GetBoundShaderState(), GetDeferredLightingVertexDeclaration<true>(), *VertexShader, *PixelShader);
 				PixelShader->SetParameters(RHICmdList, View, LightSceneInfo);
 			}
 			else
diff --git a/Engine/Source/Runtime/Renderer/Private/LightShaftRendering.cpp b/Engine/Source/Runtime/Renderer/Private/LightShaftRendering.cpp
--- a/Engine/Source/Runtime/Renderer/Private/LightShaftRendering.cpp
+++ b/Engine/Source/Runtime/Renderer/Private/LightShaftRendering.cpp
@@ -471,14 +471,14 @@ void DownsamplePass(FRHICommandListImmediate& RHICmdList, const FViewInfo& View,
 	case LightType_Directional:
 		{
 			TShaderMapRef<TDownsampleLightShaftsPixelShader<LightType_Directional, bDownsampleOcclusion> > DownsampleLightShaftsPixelShader(GetGlobalShaderMap());
-			SetGlobalBoundShaderState(RHICmdList, DownsampleDirectionalLightShaftsBoundShaderState[bDownsampleOcclusion], GFilterVertexDeclaration.VertexDeclarationRHI, *DownsampleLightShaftsVertexShader, *DownsampleLightShaftsPixelShader);
+			SetGlobalBoundShaderState(RHICmdList, View.GetFeatureLevel(), DownsampleDirectionalLightShaftsBoundShaderState[bDownsampleOcclusion], GFilterVertexDeclaration.VertexDeclarationRHI, *DownsampleLightShaftsVertexShader, *DownsampleLightShaftsPixelShader);
 			DownsampleLightShaftsPixelShader->SetParameters(RHICmdList, LightSceneInfo, View, UnusedRT);
 		}
 		break;
 	case LightType_Spot:
 		{
 			TShaderMapRef<TDownsampleLightShaftsPixelShader<LightType_Spot, bDownsampleOcclusion> > DownsampleLightShaftsPixelShader(GetGlobalShaderMap());
-			SetGlobalBoundShaderState(RHICmdList, DownsampleSpotLightShaftsBoundShaderState[bDownsampleOcclusion], GFilterVertexDeclaration.VertexDeclarationRHI, *DownsampleLightShaftsVertexShader, *DownsampleLightShaftsPixelShader);
+			SetGlobalBoundShaderState(RHICmdList, View.GetFeatureLevel(), DownsampleSpotLightShaftsBoundShaderState[bDownsampleOcclusion], GFilterVertexDeclaration.VertexDeclarationRHI, *DownsampleLightShaftsVertexShader, *DownsampleLightShaftsPixelShader);
 			DownsampleLightShaftsPixelShader->SetParameters(RHICmdList, LightSceneInfo, View, UnusedRT);
 		}
 		break;
@@ -486,7 +486,7 @@ void DownsamplePass(FRHICommandListImmediate& RHICmdList, const FViewInfo& View,
 	case LightType_Point:
 		{
 			TShaderMapRef<TDownsampleLightShaftsPixelShader<LightType_Point, bDownsampleOcclusion> > DownsampleLightShaftsPixelShader(GetGlobalShaderMap());
-			SetGlobalBoundShaderState(RHICmdList, DownsamplePointLightShaftsBoundShaderState[bDownsampleOcclusion], GFilterVertexDeclaration.VertexDeclarationRHI, *DownsampleLightShaftsVertexShader, *DownsampleLightShaftsPixelShader);
+			SetGlobalBoundShaderState(RHICmdList, View.GetFeatureLevel(), DownsamplePointLightShaftsBoundShaderState[bDownsampleOcclusion], GFilterVertexDeclaration.VertexDeclarationRHI, *DownsampleLightShaftsVertexShader, *DownsampleLightShaftsPixelShader);
 			DownsampleLightShaftsPixelShader->SetParameters(RHICmdList, LightSceneInfo, View, UnusedRT);
 		}
 		break;
@@ -616,7 +616,7 @@ void ApplyRadialBlurPasses(
 		RHICmdList.SetDepthStencilState(TStaticDepthStencilState<false, CF_Always>::GetRHI());
 
 		TShaderMapRef<FBlurLightShaftsPixelShader> BlurLightShaftsPixelShader(GetGlobalShaderMap());
-		SetGlobalBoundShaderState(RHICmdList, BlurLightShaftsBoundShaderState, GFilterVertexDeclaration.VertexDeclarationRHI, *ScreenVertexShader, *BlurLightShaftsPixelShader);
+		SetGlobalBoundShaderState(RHICmdList, View.GetFeatureLevel(), BlurLightShaftsBoundShaderState, GFilterVertexDeclaration.VertexDeclarationRHI, *ScreenVertexShader, *BlurLightShaftsPixelShader);
 
 		TRefCountPtr<IPooledRenderTarget>& EffectiveSource = PassIndex == 0 ? FirstPassSource : LightShaftsSource;
 		/// ?
@@ -662,7 +662,7 @@ void FinishOcclusionTerm(FRHICommandList& RHICmdList, const FViewInfo& View, con
 	RHICmdList.SetDepthStencilState(TStaticDepthStencilState<false, CF_Always>::GetRHI());
 
 	TShaderMapRef<FFinishOcclusionPixelShader> MaskOcclusionTermPixelShader(GetGlobalShaderMap());
-	SetGlobalBoundShaderState(RHICmdList, AccumulateTermBoundShaderState, GFilterVertexDeclaration.VertexDeclarationRHI, *ScreenVertexShader, *MaskOcclusionTermPixelShader);
+	SetGlobalBoundShaderState(RHICmdList, View.GetFeatureLevel(), AccumulateTermBoundShaderState, GFilterVertexDeclaration.VertexDeclarationRHI, *ScreenVertexShader, *MaskOcclusionTermPixelShader);
 	/// ?
 	MaskOcclusionTermPixelShader->SetParameters(RHICmdList, LightSceneInfo, View, LightShaftsSource);
 
@@ -847,7 +847,7 @@ void ApplyLightShaftBloom(FRHICommandListImmediate& RHICmdList, const FViewInfo&
 	TShaderMapRef<FScreenVS> ScreenVertexShader(GetGlobalShaderMap());
 	TShaderMapRef<FApplyLightShaftsPixelShader> ApplyLightShaftsPixelShader(GetGlobalShaderMap());
 
-	SetGlobalBoundShaderState(RHICmdList, ApplyLightShaftsBoundShaderState, GFilterVertexDeclaration.VertexDeclarationRHI, *ScreenVertexShader, *ApplyLightShaftsPixelShader);
+	SetGlobalBoundShaderState(RHICmdList, View.GetFeatureLevel(), ApplyLightShaftsBoundShaderState, GFilterVertexDeclaration.VertexDeclarationRHI, *ScreenVertexShader, *ApplyLightShaftsPixelShader);
 
 	/// ?
 	ApplyLightShaftsPixelShader->SetParameters(RHICmdList, View, LightShaftsSource);
diff --git a/Engine/Source/Runtime/Renderer/Private/PostProcess/PostProcessAA.cpp b/Engine/Source/Runtime/Renderer/Private/PostProcess/PostProcessAA.cpp
--- a/Engine/Source/Runtime/Renderer/Private/PostProcess/PostProcessAA.cpp
+++ b/Engine/Source/Runtime/Renderer/Private/PostProcess/PostProcessAA.cpp
@@ -221,7 +221,7 @@ static void SetShaderTemplAA(const FRenderingCompositePassContext& Context)
 	static FGlobalBoundShaderState BoundShaderState;
 	
 
-	SetGlobalBoundShaderState(Context.RHICmdList, BoundShaderState, GFilterVertexDeclaration.VertexDeclarationRHI, *VertexShader, *PixelShader);
+	SetGlobalBoundShaderState(Context.RHICmdList, Context.GetFeatureLevel(), BoundShaderState, GFilterVertexDeclaration.VertexDeclarationRHI, *VertexShader, *PixelShader);
 
 	PixelShader->SetParameters(Context);
 	VertexShader->SetParameters(Context);
diff --git a/Engine/Source/Runtime/Renderer/Private/PostProcess/PostProcessBloomSetup.cpp b/Engine/Source/Runtime/Renderer/Private/PostProcess/PostProcessBloomSetup.cpp
--- a/Engine/Source/Runtime/Renderer/Private/PostProcess/PostProcessBloomSetup.cpp
+++ b/Engine/Source/Runtime/Renderer/Private/PostProcess/PostProcessBloomSetup.cpp
@@ -181,7 +181,7 @@ void FRCPassPostProcessBloomSetup::Process(FRenderingCompositePassContext& Conte
 	static FGlobalBoundShaderState BoundShaderState;
 	
 
-	SetGlobalBoundShaderState(Context.RHICmdList, BoundShaderState, GFilterVertexDeclaration.VertexDeclarationRHI, *VertexShader, *PixelShader);
+	SetGlobalBoundShaderState(Context.RHICmdList, Context.GetFeatureLevel(), BoundShaderState, GFilterVertexDeclaration.VertexDeclarationRHI, *VertexShader, *PixelShader);
 
 	VertexShader->SetVS(Context);
 	PixelShader->SetPS(Context);
diff --git a/Engine/Source/Runtime/Renderer/Private/PostProcess/PostProcessBokehDOF.cpp b/Engine/Source/Runtime/Renderer/Private/PostProcess/PostProcessBokehDOF.cpp
--- a/Engine/Source/Runtime/Renderer/Private/PostProcess/PostProcessBokehDOF.cpp
+++ b/Engine/Source/Runtime/Renderer/Private/PostProcess/PostProcessBokehDOF.cpp
@@ -192,7 +192,7 @@ void FRCPassPostProcessVisualizeDOF::Process(FRenderingCompositePassContext& Con
 
 		static FGlobalBoundShaderState BoundShaderState;
 
-		SetGlobalBoundShaderState(Context.RHICmdList, BoundShaderState, GFilterVertexDeclaration.VertexDeclarationRHI, *VertexShader, *PixelShader);
+		SetGlobalBoundShaderState(Context.RHICmdList, Context.GetFeatureLevel(), BoundShaderState, GFilterVertexDeclaration.VertexDeclarationRHI, *VertexShader, *PixelShader);
 
 		VertexShader->SetParameters(Context);
 		PixelShader->SetParameters(Context);
@@ -350,7 +350,7 @@ void FRCPassPostProcessBokehDOFSetup::Process(FRenderingCompositePassContext& Co
 
 		static FGlobalBoundShaderState BoundShaderState;
 
-		SetGlobalBoundShaderState(Context.RHICmdList, BoundShaderState, GFilterVertexDeclaration.VertexDeclarationRHI, *VertexShader, *PixelShader);
+		SetGlobalBoundShaderState(Context.RHICmdList, Context.GetFeatureLevel(), BoundShaderState, GFilterVertexDeclaration.VertexDeclarationRHI, *VertexShader, *PixelShader);
 
 		VertexShader->SetParameters(Context);
 		PixelShader->SetParameters(Context);
@@ -565,7 +565,7 @@ void FRCPassPostProcessBokehDOF::SetShaderTempl(const FRenderingCompositePassCon
 	static FGlobalBoundShaderState BoundShaderState;
 	
 
-	SetGlobalBoundShaderState(Context.RHICmdList, BoundShaderState, GFilterVertexDeclaration.VertexDeclarationRHI, *VertexShader, *PixelShader);
+	SetGlobalBoundShaderState(Context.RHICmdList, Context.GetFeatureLevel(), BoundShaderState, GFilterVertexDeclaration.VertexDeclarationRHI, *VertexShader, *PixelShader);
 
 	VertexShader->SetParameters(Context, TileCount, TileSize, PixelKernelSize, LeftTop);
 	PixelShader->SetParameters(Context, PixelKernelSize);
diff --git a/Engine/Source/Runtime/Renderer/Private/PostProcess/PostProcessBokehDOFRecombine.cpp b/Engine/Source/Runtime/Renderer/Private/PostProcess/PostProcessBokehDOFRecombine.cpp
--- a/Engine/Source/Runtime/Renderer/Private/PostProcess/PostProcessBokehDOFRecombine.cpp
+++ b/Engine/Source/Runtime/Renderer/Private/PostProcess/PostProcessBokehDOFRecombine.cpp
@@ -103,7 +103,7 @@ void FRCPassPostProcessBokehDOFRecombine::SetShader(const FRenderingCompositePas
 	static FGlobalBoundShaderState BoundShaderState;
 	
 
-	SetGlobalBoundShaderState(Context.RHICmdList, BoundShaderState, GFilterVertexDeclaration.VertexDeclarationRHI, *VertexShader, *PixelShader);
+	SetGlobalBoundShaderState(Context.RHICmdList, Context.GetFeatureLevel(), BoundShaderState, GFilterVertexDeclaration.VertexDeclarationRHI, *VertexShader, *PixelShader);
 
 	PixelShader->SetParameters(Context);
 	VertexShader->SetParameters(Context);
diff --git a/Engine/Source/Runtime/Renderer/Private/PostProcess/PostProcessBusyWait.cpp b/Engine/Source/Runtime/Renderer/Private/PostProcess/PostProcessBusyWait.cpp
--- a/Engine/Source/Runtime/Renderer/Private/PostProcess/PostProcessBusyWait.cpp
+++ b/Engine/Source/Runtime/Renderer/Private/PostProcess/PostProcessBusyWait.cpp
@@ -107,7 +107,7 @@ void FRCPassPostProcessBusyWait::Process(FRenderingCompositePassContext& Context
 
 	static FGlobalBoundShaderState BoundShaderState;
 
-	SetGlobalBoundShaderState(Context.RHICmdList, BoundShaderState, GFilterVertexDeclaration.VertexDeclarationRHI, *VertexShader, *PixelShader);
+	SetGlobalBoundShaderState(Context.RHICmdList, Context.GetFeatureLevel(), BoundShaderState, GFilterVertexDeclaration.VertexDeclarationRHI, *VertexShader, *PixelShader);
 
 	PixelShader->SetPS(Context);
 
diff --git a/Engine/Source/Runtime/Renderer/Private/PostProcess/PostProcessCombineLUTs.cpp b/Engine/Source/Runtime/Renderer/Private/PostProcess/PostProcessCombineLUTs.cpp
--- a/Engine/Source/Runtime/Renderer/Private/PostProcess/PostProcessCombineLUTs.cpp
+++ b/Engine/Source/Runtime/Renderer/Private/PostProcess/PostProcessCombineLUTs.cpp
@@ -226,12 +226,12 @@ void SetLUTBlenderShader(FRenderingCompositePassContext& Context, uint32 BlendCo
 	}
 #undef CASE_COUNT
 	check(LocalBoundShaderState != NULL);
-	if(UseVolumeTextureLUT(GRHIShaderPlatform))
+	if(UseVolumeTextureLUT(Context.View.GetShaderPlatform()))
 	{
 		TShaderMapRef<FWriteToSliceVS> VertexShader(GetGlobalShaderMap());
 		TShaderMapRef<FWriteToSliceGS> GeometryShader(GetGlobalShaderMap());
 
-		SetGlobalBoundShaderState(Context.RHICmdList, *LocalBoundShaderState, GScreenVertexDeclaration.VertexDeclarationRHI, *VertexShader, LocalPixelShader, *GeometryShader);
+		SetGlobalBoundShaderState(Context.RHICmdList, Context.GetFeatureLevel(), *LocalBoundShaderState, GScreenVertexDeclaration.VertexDeclarationRHI, *VertexShader, LocalPixelShader, *GeometryShader);
 
 		VertexShader->SetParameters(Context.RHICmdList, VolumeBounds, VolumeBounds.MaxX - VolumeBounds.MinX);
 		GeometryShader->SetParameters(Context.RHICmdList, VolumeBounds);
@@ -240,7 +240,7 @@ void SetLUTBlenderShader(FRenderingCompositePassContext& Context, uint32 BlendCo
 	{
 		TShaderMapRef<FPostProcessVS> VertexShader(GetGlobalShaderMap());
 
-		SetGlobalBoundShaderState(Context.RHICmdList, *LocalBoundShaderState, GFilterVertexDeclaration.VertexDeclarationRHI, *VertexShader, LocalPixelShader);
+		SetGlobalBoundShaderState(Context.RHICmdList, Context.GetFeatureLevel(), *LocalBoundShaderState, GFilterVertexDeclaration.VertexDeclarationRHI, *VertexShader, LocalPixelShader);
 
 		VertexShader->SetParameters(Context);
 	}
@@ -396,7 +396,7 @@ void FRCPassPostProcessCombineLUTs::Process(FRenderingCompositePassContext& Cont
 	}
 
 	// for a 3D texture, the viewport is 16x16 (per slice), for a 2D texture, it's unwrapped to 256x16
-	FIntPoint DestSize(UseVolumeTextureLUT(GRHIShaderPlatform) ? 16 : 256, 16);
+	FIntPoint DestSize(UseVolumeTextureLUT(ShaderPlatform) ? 16 : 256, 16);
 
 	const FSceneRenderTargetItem& DestRenderTarget = PassOutputs[0].RequestSurface(Context);
 
@@ -413,7 +413,7 @@ void FRCPassPostProcessCombineLUTs::Process(FRenderingCompositePassContext& Cont
 
 	SetLUTBlenderShader(Context, LocalCount, LocalTextures, LocalWeights, VolumeBounds);
 
-	if(UseVolumeTextureLUT(GRHIShaderPlatform))
+	if (UseVolumeTextureLUT(ShaderPlatform))
 	{
 		// use volume texture 16x16x16
 		RasterizeToVolumeTexture(Context.RHICmdList, VolumeBounds);
@@ -442,7 +442,7 @@ FPooledRenderTargetDesc FRCPassPostProcessCombineLUTs::ComputeOutputDesc(EPassOu
 {
 	FPooledRenderTargetDesc Ret = FPooledRenderTargetDesc::Create2DDesc(FIntPoint(256, 16), PF_B8G8R8A8, TexCreate_None, TexCreate_RenderTargetable | TexCreate_ShaderResource, false);
 
-	if(UseVolumeTextureLUT(GRHIShaderPlatform))
+	if(UseVolumeTextureLUT(ShaderPlatform))
 	{
 		Ret.Extent = FIntPoint(16, 16);
 		Ret.Depth = 16;
diff --git a/Engine/Source/Runtime/Renderer/Private/PostProcess/PostProcessCombineLUTs.h b/Engine/Source/Runtime/Renderer/Private/PostProcess/PostProcessCombineLUTs.h
--- a/Engine/Source/Runtime/Renderer/Private/PostProcess/PostProcessCombineLUTs.h
+++ b/Engine/Source/Runtime/Renderer/Private/PostProcess/PostProcessCombineLUTs.h
@@ -13,6 +13,11 @@ class FRCPassPostProcessCombineLUTs : public TRenderingCompositePassBase<0, 1>
 {
 public:
 	// interface FRenderingCompositePass ---------
+	FRCPassPostProcessCombineLUTs(EShaderPlatform InShaderPlatform)
+	: ShaderPlatform(InShaderPlatform)
+	{
+		
+	}
 
 	virtual void Process(FRenderingCompositePassContext& Context);
 	virtual void Release() override { delete this; }
@@ -24,6 +29,10 @@ public:
 	uint32 FindIndex(const FFinalPostProcessSettings& Settings, UTexture* Tex) const;
 
 	static bool IsColorGradingLUTNeeded(const FRenderingCompositePassContext* RESTRICT Context);
+
+private:
+
+	EShaderPlatform ShaderPlatform;
 };
 
 
diff --git a/Engine/Source/Runtime/Renderer/Private/PostProcess/PostProcessCompositeEditorPrimitives.cpp b/Engine/Source/Runtime/Renderer/Private/PostProcess/PostProcessCompositeEditorPrimitives.cpp
--- a/Engine/Source/Runtime/Renderer/Private/PostProcess/PostProcessCompositeEditorPrimitives.cpp
+++ b/Engine/Source/Runtime/Renderer/Private/PostProcess/PostProcessCompositeEditorPrimitives.cpp
@@ -170,7 +170,7 @@ static void SetCompositePrimitivesShaderTempl(const FRenderingCompositePassConte
 	static FGlobalBoundShaderState BoundShaderState;
 	
 
-	SetGlobalBoundShaderState(Context.RHICmdList, BoundShaderState, GFilterVertexDeclaration.VertexDeclarationRHI, *VertexShader, *PixelShader);
+	SetGlobalBoundShaderState(Context.RHICmdList, Context.GetFeatureLevel(), BoundShaderState, GFilterVertexDeclaration.VertexDeclarationRHI, *VertexShader, *PixelShader);
 
 	VertexShader->SetParameters(Context);
 	PixelShader->SetParameters(Context);
@@ -302,7 +302,9 @@ void FRCPassPostProcessCompositeEditorPrimitives::RenderPrimitivesToComposite(FR
 	// most objects should be occluded by the existing scene so we do a manual depth test in the shader
 	bool bDepthTest = true;
 
-	const bool bNeedToSwitchVerticalAxis = RHINeedsToSwitchVerticalAxis(GRHIShaderPlatform);
+	const auto FeatureLevel = View.GetFeatureLevel();
+	const auto ShaderPlatform = GShaderPlatformForFeatureLevel[FeatureLevel];
+	const bool bNeedToSwitchVerticalAxis = RHINeedsToSwitchVerticalAxis(ShaderPlatform);
 	FTexture2DRHIRef SceneDepth = GSceneRenderTargets.GetSceneDepthTexture();
 
 	const bool bUseGetMeshElements = ShouldUseGetDynamicMeshElements();
@@ -371,7 +373,7 @@ void FRCPassPostProcessCompositeEditorPrimitives::RenderPrimitivesToComposite(FR
 	DrawViewElements<FBasePassOpaqueDrawingPolicyFactory>(RHICmdList, View, FBasePassOpaqueDrawingPolicyFactory::ContextType(bDepthTest, ESceneRenderTargetsMode::SetTextures), SDPG_World, false);
 
 	// Draw the view's batched simple elements(lines, sprites, etc).
-	View.BatchedViewElements.Draw(RHICmdList, bNeedToSwitchVerticalAxis, View.ViewProjectionMatrix, View.ViewRect.Width(), View.ViewRect.Height(), false, 1.0f, &View, SceneDepth);
+	View.BatchedViewElements.Draw(RHICmdList, FeatureLevel, bNeedToSwitchVerticalAxis, View.ViewProjectionMatrix, View.ViewRect.Width(), View.ViewRect.Height(), false, 1.0f, &View, SceneDepth);
 
 	// Draw foreground objects. Draw twice, once without depth testing to bring them into the foreground and again to depth test against themselves
 	{
@@ -382,13 +384,13 @@ void FRCPassPostProcessCompositeEditorPrimitives::RenderPrimitivesToComposite(FR
 
 		DrawViewElements<FBasePassOpaqueDrawingPolicyFactory>(RHICmdList, View, FBasePassOpaqueDrawingPolicyFactory::ContextType(bDepthTest, ESceneRenderTargetsMode::SetTextures), SDPG_Foreground, false);
 
-		View.TopBatchedViewElements.Draw(RHICmdList, bNeedToSwitchVerticalAxis, View.ViewProjectionMatrix, View.ViewRect.Width(), View.ViewRect.Height(), false);
+		View.TopBatchedViewElements.Draw(RHICmdList, FeatureLevel, bNeedToSwitchVerticalAxis, View.ViewProjectionMatrix, View.ViewRect.Width(), View.ViewRect.Height(), false);
 
 		// Note, this is a reversed Z depth surface, using CF_GreaterEqual.
 		RHICmdList.SetDepthStencilState(TStaticDepthStencilState<true, CF_GreaterEqual>::GetRHI());
 
 		DrawViewElements<FBasePassOpaqueDrawingPolicyFactory>(RHICmdList, View, FBasePassOpaqueDrawingPolicyFactory::ContextType(bDepthTest, ESceneRenderTargetsMode::SetTextures), SDPG_Foreground, false);
 
-		View.TopBatchedViewElements.Draw(RHICmdList, bNeedToSwitchVerticalAxis, View.ViewProjectionMatrix, View.ViewRect.Width(), View.ViewRect.Height(), false);
+		View.TopBatchedViewElements.Draw(RHICmdList, FeatureLevel, bNeedToSwitchVerticalAxis, View.ViewProjectionMatrix, View.ViewRect.Width(), View.ViewRect.Height(), false);
 	}
 }
diff --git a/Engine/Source/Runtime/Renderer/Private/PostProcess/PostProcessDOF.cpp b/Engine/Source/Runtime/Renderer/Private/PostProcess/PostProcessDOF.cpp
--- a/Engine/Source/Runtime/Renderer/Private/PostProcess/PostProcessDOF.cpp
+++ b/Engine/Source/Runtime/Renderer/Private/PostProcess/PostProcessDOF.cpp
@@ -137,7 +137,7 @@ void FRCPassPostProcessDOFSetup::Process(FRenderingCompositePassContext& Context
 		
 
 		TShaderMapRef< FPostProcessDOFSetupPS<1> > PixelShader(GetGlobalShaderMap());
-		SetGlobalBoundShaderState(Context.RHICmdList, BoundShaderState, GFilterVertexDeclaration.VertexDeclarationRHI, *VertexShader, *PixelShader);
+		SetGlobalBoundShaderState(Context.RHICmdList, Context.GetFeatureLevel(), BoundShaderState, GFilterVertexDeclaration.VertexDeclarationRHI, *VertexShader, *PixelShader);
 		
 		PixelShader->SetParameters(Context);
 	}
@@ -146,7 +146,7 @@ void FRCPassPostProcessDOFSetup::Process(FRenderingCompositePassContext& Context
 		static FGlobalBoundShaderState BoundShaderState;
 		
 		TShaderMapRef< FPostProcessDOFSetupPS<0> > PixelShader(GetGlobalShaderMap());
-		SetGlobalBoundShaderState(Context.RHICmdList, BoundShaderState, GFilterVertexDeclaration.VertexDeclarationRHI, *VertexShader, *PixelShader);
+		SetGlobalBoundShaderState(Context.RHICmdList, Context.GetFeatureLevel(), BoundShaderState, GFilterVertexDeclaration.VertexDeclarationRHI, *VertexShader, *PixelShader);
 
 		PixelShader->SetParameters(Context);
 	}
@@ -301,7 +301,7 @@ void FRCPassPostProcessDOFRecombine::Process(FRenderingCompositePassContext& Con
 		static FGlobalBoundShaderState BoundShaderState;
 		
 		TShaderMapRef< FPostProcessDOFRecombinePS<1> > PixelShader(GetGlobalShaderMap());
-		SetGlobalBoundShaderState(Context.RHICmdList, BoundShaderState, GFilterVertexDeclaration.VertexDeclarationRHI, *VertexShader, *PixelShader);
+		SetGlobalBoundShaderState(Context.RHICmdList, Context.GetFeatureLevel(), BoundShaderState, GFilterVertexDeclaration.VertexDeclarationRHI, *VertexShader, *PixelShader);
 		PixelShader->SetParameters(Context);
 	}
 	else
@@ -309,7 +309,7 @@ void FRCPassPostProcessDOFRecombine::Process(FRenderingCompositePassContext& Con
 		static FGlobalBoundShaderState BoundShaderState;
 		
 		TShaderMapRef< FPostProcessDOFRecombinePS<0> > PixelShader(GetGlobalShaderMap());
-		SetGlobalBoundShaderState(Context.RHICmdList, BoundShaderState, GFilterVertexDeclaration.VertexDeclarationRHI, *VertexShader, *PixelShader);
+		SetGlobalBoundShaderState(Context.RHICmdList, Context.GetFeatureLevel(), BoundShaderState, GFilterVertexDeclaration.VertexDeclarationRHI, *VertexShader, *PixelShader);
 		PixelShader->SetParameters(Context);
 	}
 
diff --git a/Engine/Source/Runtime/Renderer/Private/PostProcess/PostProcessDownsample.cpp b/Engine/Source/Runtime/Renderer/Private/PostProcess/PostProcessDownsample.cpp
--- a/Engine/Source/Runtime/Renderer/Private/PostProcess/PostProcessDownsample.cpp
+++ b/Engine/Source/Runtime/Renderer/Private/PostProcess/PostProcessDownsample.cpp
@@ -144,7 +144,7 @@ void FRCPassPostProcessDownsample::SetShader(const FRenderingCompositePassContex
 	static FGlobalBoundShaderState BoundShaderState;
 	
 
-	SetGlobalBoundShaderState(Context.RHICmdList, BoundShaderState, GFilterVertexDeclaration.VertexDeclarationRHI, *VertexShader, *PixelShader);
+	SetGlobalBoundShaderState(Context.RHICmdList, Context.GetFeatureLevel(), BoundShaderState, GFilterVertexDeclaration.VertexDeclarationRHI, *VertexShader, *PixelShader);
 
 	PixelShader->SetParameters(Context);
 	VertexShader->SetParameters(Context);
diff --git a/Engine/Source/Runtime/Renderer/Private/PostProcess/PostProcessEyeAdaptation.cpp b/Engine/Source/Runtime/Renderer/Private/PostProcess/PostProcessEyeAdaptation.cpp
--- a/Engine/Source/Runtime/Renderer/Private/PostProcess/PostProcessEyeAdaptation.cpp
+++ b/Engine/Source/Runtime/Renderer/Private/PostProcess/PostProcessEyeAdaptation.cpp
@@ -96,7 +96,7 @@ void FRCPassPostProcessEyeAdaptation::Process(FRenderingCompositePassContext& Co
 	static FGlobalBoundShaderState BoundShaderState;
 	
 
-	SetGlobalBoundShaderState(Context.RHICmdList, BoundShaderState, GFilterVertexDeclaration.VertexDeclarationRHI, *VertexShader, *PixelShader);
+	SetGlobalBoundShaderState(Context.RHICmdList, Context.GetFeatureLevel(), BoundShaderState, GFilterVertexDeclaration.VertexDeclarationRHI, *VertexShader, *PixelShader);
 
 	PixelShader->SetPS(Context);
 
diff --git a/Engine/Source/Runtime/Renderer/Private/PostProcess/PostProcessGBufferHints.cpp b/Engine/Source/Runtime/Renderer/Private/PostProcess/PostProcessGBufferHints.cpp
--- a/Engine/Source/Runtime/Renderer/Private/PostProcess/PostProcessGBufferHints.cpp
+++ b/Engine/Source/Runtime/Renderer/Private/PostProcess/PostProcessGBufferHints.cpp
@@ -119,7 +119,7 @@ void FRCPassPostProcessGBufferHints::Process(FRenderingCompositePassContext& Con
 	static FGlobalBoundShaderState BoundShaderState;
 	
 
-	SetGlobalBoundShaderState(Context.RHICmdList, BoundShaderState, GFilterVertexDeclaration.VertexDeclarationRHI, *VertexShader, *PixelShader);
+	SetGlobalBoundShaderState(Context.RHICmdList, Context.GetFeatureLevel(), BoundShaderState, GFilterVertexDeclaration.VertexDeclarationRHI, *VertexShader, *PixelShader);
 
 	PixelShader->SetPS(Context);
 
@@ -157,7 +157,7 @@ void FRCPassPostProcessGBufferHints::Process(FRenderingCompositePassContext& Con
 		}
 	} TempRenderTarget(View, (const FTexture2DRHIRef&)DestRenderTarget.TargetableTexture);
 
-	FCanvas Canvas(&TempRenderTarget, NULL, ViewFamily.CurrentRealTime, ViewFamily.CurrentWorldTime, ViewFamily.DeltaWorldTime);
+	FCanvas Canvas(&TempRenderTarget, NULL, ViewFamily.CurrentRealTime, ViewFamily.CurrentWorldTime, ViewFamily.DeltaWorldTime, View.GetFeatureLevel());
 
 	float X = 30;
 	float Y = 8;
diff --git a/Engine/Source/Runtime/Renderer/Private/PostProcess/PostProcessHMD.cpp b/Engine/Source/Runtime/Renderer/Private/PostProcess/PostProcessHMD.cpp
--- a/Engine/Source/Runtime/Renderer/Private/PostProcess/PostProcessHMD.cpp
+++ b/Engine/Source/Runtime/Renderer/Private/PostProcess/PostProcessHMD.cpp
@@ -223,7 +223,7 @@ void FRCPassPostProcessHMD::Process(FRenderingCompositePassContext& Context)
 		TShaderMapRef<FPostProcessHMDPS<false> > PixelShader(GetGlobalShaderMap());
 		static FGlobalBoundShaderState BoundShaderState;
 		
-		SetGlobalBoundShaderState(Context.RHICmdList, BoundShaderState, GDistortionVertexDeclaration.VertexDeclarationRHI, *VertexShader, *PixelShader);
+		SetGlobalBoundShaderState(Context.RHICmdList, Context.GetFeatureLevel(), BoundShaderState, GDistortionVertexDeclaration.VertexDeclarationRHI, *VertexShader, *PixelShader);
 		VertexShader->SetVS(Context, View.StereoPass);
 		PixelShader->SetPS(Context, SrcRect, SrcSize, View.StereoPass, QuadTexTransform);
 	}
diff --git a/Engine/Source/Runtime/Renderer/Private/PostProcess/PostProcessHistogramReduce.cpp b/Engine/Source/Runtime/Renderer/Private/PostProcess/PostProcessHistogramReduce.cpp
--- a/Engine/Source/Runtime/Renderer/Private/PostProcess/PostProcessHistogramReduce.cpp
+++ b/Engine/Source/Runtime/Renderer/Private/PostProcess/PostProcessHistogramReduce.cpp
@@ -119,7 +119,7 @@ void FRCPassPostProcessHistogramReduce::Process(FRenderingCompositePassContext&
 	static FGlobalBoundShaderState BoundShaderState;
 	
 
-	SetGlobalBoundShaderState(Context.RHICmdList, BoundShaderState, GFilterVertexDeclaration.VertexDeclarationRHI, *VertexShader, *PixelShader);
+	SetGlobalBoundShaderState(Context.RHICmdList, Context.GetFeatureLevel(), BoundShaderState, GFilterVertexDeclaration.VertexDeclarationRHI, *VertexShader, *PixelShader);
 
 	// we currently assume the input is half res, one full res pixel less to avoid getting bilinear filtered input
 	FIntPoint GatherExtent = (View.ViewRect.Size() - FIntPoint(1, 1)) / 2;
diff --git a/Engine/Source/Runtime/Renderer/Private/PostProcess/PostProcessLensBlur.cpp b/Engine/Source/Runtime/Renderer/Private/PostProcess/PostProcessLensBlur.cpp
--- a/Engine/Source/Runtime/Renderer/Private/PostProcess/PostProcessLensBlur.cpp
+++ b/Engine/Source/Runtime/Renderer/Private/PostProcess/PostProcessLensBlur.cpp
@@ -200,7 +200,7 @@ void FRCPassPostProcessLensBlur::Process(FRenderingCompositePassContext& Context
 
 	static FGlobalBoundShaderState BoundShaderState;
 	
-	SetGlobalBoundShaderState(Context.RHICmdList, BoundShaderState, GFilterVertexDeclaration.VertexDeclarationRHI, *VertexShader, *PixelShader);
+	SetGlobalBoundShaderState(Context.RHICmdList, Context.GetFeatureLevel(), BoundShaderState, GFilterVertexDeclaration.VertexDeclarationRHI, *VertexShader, *PixelShader);
 
 	uint32 TileSize = 1;
 
diff --git a/Engine/Source/Runtime/Renderer/Private/PostProcess/PostProcessLensFlares.cpp b/Engine/Source/Runtime/Renderer/Private/PostProcess/PostProcessLensFlares.cpp
--- a/Engine/Source/Runtime/Renderer/Private/PostProcess/PostProcessLensFlares.cpp
+++ b/Engine/Source/Runtime/Renderer/Private/PostProcess/PostProcessLensFlares.cpp
@@ -162,7 +162,7 @@ void FRCPassPostProcessLensFlares::Process(FRenderingCompositePassContext& Conte
 
 		static FGlobalBoundShaderState BoundShaderState;
 		
-		SetGlobalBoundShaderState(Context.RHICmdList, BoundShaderState, GFilterVertexDeclaration.VertexDeclarationRHI, *VertexShader, *PixelShader);
+		SetGlobalBoundShaderState(Context.RHICmdList, Context.GetFeatureLevel(), BoundShaderState, GFilterVertexDeclaration.VertexDeclarationRHI, *VertexShader, *PixelShader);
 
 		VertexShader->SetParameters(Context);
 		PixelShader->SetParameters(Context);
@@ -189,7 +189,7 @@ void FRCPassPostProcessLensFlares::Process(FRenderingCompositePassContext& Conte
 
 		static FGlobalBoundShaderState BoundShaderState;
 		
-		SetGlobalBoundShaderState(Context.RHICmdList, BoundShaderState, GFilterVertexDeclaration.VertexDeclarationRHI, *VertexShader, *PixelShader);
+		SetGlobalBoundShaderState(Context.RHICmdList, Context.GetFeatureLevel(), BoundShaderState, GFilterVertexDeclaration.VertexDeclarationRHI, *VertexShader, *PixelShader);
 
 		FVector2D TexScaleValue = FVector2D(TexSize2) / ViewSize2;
 
diff --git a/Engine/Source/Runtime/Renderer/Private/PostProcess/PostProcessMobile.cpp b/Engine/Source/Runtime/Renderer/Private/PostProcess/PostProcessMobile.cpp
--- a/Engine/Source/Runtime/Renderer/Private/PostProcess/PostProcessMobile.cpp
+++ b/Engine/Source/Runtime/Renderer/Private/PostProcess/PostProcessMobile.cpp
@@ -139,7 +139,7 @@ static void BloomSetup_SetShader(const FRenderingCompositePassContext& Context)
 
 	static FGlobalBoundShaderState BoundShaderState;
 	
-	SetGlobalBoundShaderState(Context.RHICmdList, BoundShaderState, GFilterVertexDeclaration.VertexDeclarationRHI, *VertexShader, *PixelShader);
+	SetGlobalBoundShaderState(Context.RHICmdList, Context.GetFeatureLevel(), BoundShaderState, GFilterVertexDeclaration.VertexDeclarationRHI, *VertexShader, *PixelShader);
 
 	VertexShader->SetVS(Context);
 	PixelShader->SetPS(Context);
@@ -360,7 +360,7 @@ void FRCPassPostProcessBloomSetupSmallES2::SetShader(const FRenderingCompositePa
 
 	static FGlobalBoundShaderState BoundShaderState;
 	
-	SetGlobalBoundShaderState(Context.RHICmdList, BoundShaderState, GFilterVertexDeclaration.VertexDeclarationRHI, *VertexShader, *PixelShader);
+	SetGlobalBoundShaderState(Context.RHICmdList, Context.GetFeatureLevel(), BoundShaderState, GFilterVertexDeclaration.VertexDeclarationRHI, *VertexShader, *PixelShader);
 
 	VertexShader->SetVS(Context);
 	PixelShader->SetPS(Context);
@@ -573,7 +573,7 @@ void FRCPassPostProcessBloomDownES2::Process(FRenderingCompositePassContext& Con
 	static FGlobalBoundShaderState BoundShaderState;
 	
 
-	SetGlobalBoundShaderState(Context.RHICmdList, BoundShaderState, GFilterVertexDeclaration.VertexDeclarationRHI, *VertexShader, *PixelShader);
+	SetGlobalBoundShaderState(Context.RHICmdList, Context.GetFeatureLevel(), BoundShaderState, GFilterVertexDeclaration.VertexDeclarationRHI, *VertexShader, *PixelShader);
 
 	VertexShader->SetVS(Context, Scale);
 	PixelShader->SetPS(Context);
@@ -741,7 +741,7 @@ void FRCPassPostProcessBloomUpES2::Process(FRenderingCompositePassContext& Conte
 	static FGlobalBoundShaderState BoundShaderState;
 	
 
-	SetGlobalBoundShaderState(Context.RHICmdList, BoundShaderState, GFilterVertexDeclaration.VertexDeclarationRHI, *VertexShader, *PixelShader);
+	SetGlobalBoundShaderState(Context.RHICmdList, Context.GetFeatureLevel(), BoundShaderState, GFilterVertexDeclaration.VertexDeclarationRHI, *VertexShader, *PixelShader);
 
 	// The 1/8 factor is because bloom is using 8 taps in the filter.
 	VertexShader->SetVS(Context, FVector2D(ScaleAB.X, ScaleAB.Y));
@@ -907,7 +907,7 @@ static void SunMask_SetShader(const FRenderingCompositePassContext& Context)
 
 	static FGlobalBoundShaderState BoundShaderState;
 	
-	SetGlobalBoundShaderState(Context.RHICmdList, BoundShaderState, GFilterVertexDeclaration.VertexDeclarationRHI, *VertexShader, *PixelShader);
+	SetGlobalBoundShaderState(Context.RHICmdList, Context.GetFeatureLevel(), BoundShaderState, GFilterVertexDeclaration.VertexDeclarationRHI, *VertexShader, *PixelShader);
 
 	VertexShader->SetVS(Context);
 	PixelShader->SetPS(Context);
@@ -1138,7 +1138,7 @@ static void SunAlpha_SetShader(const FRenderingCompositePassContext& Context)
 
 	static FGlobalBoundShaderState BoundShaderState;
 
-	SetGlobalBoundShaderState(Context.RHICmdList, BoundShaderState, GFilterVertexDeclaration.VertexDeclarationRHI, *VertexShader, *PixelShader);
+	SetGlobalBoundShaderState(Context.RHICmdList, Context.GetFeatureLevel(), BoundShaderState, GFilterVertexDeclaration.VertexDeclarationRHI, *VertexShader, *PixelShader);
 
 	VertexShader->SetVS(Context);
 	PixelShader->SetPS(Context);
@@ -1344,7 +1344,7 @@ void FRCPassPostProcessSunBlurES2::Process(FRenderingCompositePassContext& Conte
 	static FGlobalBoundShaderState BoundShaderState;
 	
 
-	SetGlobalBoundShaderState(Context.RHICmdList, BoundShaderState, GFilterVertexDeclaration.VertexDeclarationRHI, *VertexShader, *PixelShader);
+	SetGlobalBoundShaderState(Context.RHICmdList, Context.GetFeatureLevel(), BoundShaderState, GFilterVertexDeclaration.VertexDeclarationRHI, *VertexShader, *PixelShader);
 
 	VertexShader->SetVS(Context);
 	PixelShader->SetPS(Context);
@@ -1515,7 +1515,7 @@ static void SunMerge_SetShader(const FRenderingCompositePassContext& Context)
 	static FGlobalBoundShaderState BoundShaderState;
 	
 
-	SetGlobalBoundShaderState(Context.RHICmdList, BoundShaderState, GFilterVertexDeclaration.VertexDeclarationRHI, *VertexShader, *PixelShader);
+	SetGlobalBoundShaderState(Context.RHICmdList, Context.GetFeatureLevel(), BoundShaderState, GFilterVertexDeclaration.VertexDeclarationRHI, *VertexShader, *PixelShader);
 
 	VertexShader->SetVS(Context);
 	PixelShader->SetPS(Context);
@@ -1726,7 +1726,7 @@ void FRCPassPostProcessSunMergeSmallES2::SetShader(const FRenderingCompositePass
 	static FGlobalBoundShaderState BoundShaderState;
 	
 
-	SetGlobalBoundShaderState(Context.RHICmdList, BoundShaderState, GFilterVertexDeclaration.VertexDeclarationRHI, *VertexShader, *PixelShader);
+	SetGlobalBoundShaderState(Context.RHICmdList, Context.GetFeatureLevel(), BoundShaderState, GFilterVertexDeclaration.VertexDeclarationRHI, *VertexShader, *PixelShader);
 
 	VertexShader->SetVS(Context);
 	PixelShader->SetPS(Context);
@@ -1916,7 +1916,7 @@ static void DofDown_SetShader(const FRenderingCompositePassContext& Context)
 	static FGlobalBoundShaderState BoundShaderState;
 	
 
-	SetGlobalBoundShaderState(Context.RHICmdList, BoundShaderState, GFilterVertexDeclaration.VertexDeclarationRHI, *VertexShader, *PixelShader);
+	SetGlobalBoundShaderState(Context.RHICmdList, Context.GetFeatureLevel(), BoundShaderState, GFilterVertexDeclaration.VertexDeclarationRHI, *VertexShader, *PixelShader);
 
 	VertexShader->SetVS(Context);
 	PixelShader->SetPS(Context);
@@ -2124,7 +2124,7 @@ static void DofNear_SetShader(const FRenderingCompositePassContext& Context)
 	static FGlobalBoundShaderState BoundShaderState;
 	
 
-	SetGlobalBoundShaderState(Context.RHICmdList, BoundShaderState, GFilterVertexDeclaration.VertexDeclarationRHI, *VertexShader, *PixelShader);
+	SetGlobalBoundShaderState(Context.RHICmdList, Context.GetFeatureLevel(), BoundShaderState, GFilterVertexDeclaration.VertexDeclarationRHI, *VertexShader, *PixelShader);
 
 	VertexShader->SetVS(Context);
 	PixelShader->SetPS(Context);
@@ -2323,7 +2323,7 @@ void FRCPassPostProcessDofBlurES2::Process(FRenderingCompositePassContext& Conte
 
 	static FGlobalBoundShaderState BoundShaderState;
 	
-	SetGlobalBoundShaderState(Context.RHICmdList, BoundShaderState, GFilterVertexDeclaration.VertexDeclarationRHI, *VertexShader, *PixelShader);
+	SetGlobalBoundShaderState(Context.RHICmdList, Context.GetFeatureLevel(), BoundShaderState, GFilterVertexDeclaration.VertexDeclarationRHI, *VertexShader, *PixelShader);
 
 	VertexShader->SetVS(Context);
 	PixelShader->SetPS(Context);
@@ -2462,7 +2462,7 @@ static void SunAvg_SetShader(const FRenderingCompositePassContext& Context)
 	static FGlobalBoundShaderState BoundShaderState;
 	
 
-	SetGlobalBoundShaderState(Context.RHICmdList, BoundShaderState, GFilterVertexDeclaration.VertexDeclarationRHI, *VertexShader, *PixelShader);
+	SetGlobalBoundShaderState(Context.RHICmdList, Context.GetFeatureLevel(), BoundShaderState, GFilterVertexDeclaration.VertexDeclarationRHI, *VertexShader, *PixelShader);
 
 	VertexShader->SetVS(Context);
 	PixelShader->SetPS(Context);
@@ -2705,7 +2705,7 @@ static void Aa_SetShader(const FRenderingCompositePassContext& Context)
 	static FGlobalBoundShaderState BoundShaderState;
 	
 
-	SetGlobalBoundShaderState(Context.RHICmdList, BoundShaderState, GFilterVertexDeclaration.VertexDeclarationRHI, *VertexShader, *PixelShader);
+	SetGlobalBoundShaderState(Context.RHICmdList, Context.GetFeatureLevel(), BoundShaderState, GFilterVertexDeclaration.VertexDeclarationRHI, *VertexShader, *PixelShader);
 
 	VertexShader->SetVS(Context);
 	PixelShader->SetPS(Context);
diff --git a/Engine/Source/Runtime/Renderer/Private/PostProcess/PostProcessMorpheus.cpp b/Engine/Source/Runtime/Renderer/Private/PostProcess/PostProcessMorpheus.cpp
--- a/Engine/Source/Runtime/Renderer/Private/PostProcess/PostProcessMorpheus.cpp
+++ b/Engine/Source/Runtime/Renderer/Private/PostProcess/PostProcessMorpheus.cpp
@@ -210,7 +210,7 @@ void FRCPassPostProcessMorpheus::Process(FRenderingCompositePassContext& Context
 	static FGlobalBoundShaderState BoundShaderState;
 	
 
-	SetGlobalBoundShaderState(Context.RHICmdList, BoundShaderState, GFilterVertexDeclaration.VertexDeclarationRHI, *VertexShader, *PixelShader);
+	SetGlobalBoundShaderState(Context.RHICmdList, Context.GetFeatureLevel(), BoundShaderState, GFilterVertexDeclaration.VertexDeclarationRHI, *VertexShader, *PixelShader);
 
     FMatrix QuadTexTransform;
     FMatrix QuadPosTransform = FMatrix::Identity;
diff --git a/Engine/Source/Runtime/Renderer/Private/PostProcess/PostProcessMotionBlur.cpp b/Engine/Source/Runtime/Renderer/Private/PostProcess/PostProcessMotionBlur.cpp
--- a/Engine/Source/Runtime/Renderer/Private/PostProcess/PostProcessMotionBlur.cpp
+++ b/Engine/Source/Runtime/Renderer/Private/PostProcess/PostProcessMotionBlur.cpp
@@ -141,7 +141,7 @@ void FRCPassPostProcessMotionBlurSetup::Process(FRenderingCompositePassContext&
 		static FGlobalBoundShaderState BoundShaderState;
 		
 
-		SetGlobalBoundShaderState(Context.RHICmdList, BoundShaderState, GFilterVertexDeclaration.VertexDeclarationRHI, *VertexShader, *PixelShader);
+		SetGlobalBoundShaderState(Context.RHICmdList, Context.GetFeatureLevel(), BoundShaderState, GFilterVertexDeclaration.VertexDeclarationRHI, *VertexShader, *PixelShader);
 
 		PixelShader->SetParameters(Context);
 		VertexShader->SetParameters(Context);
@@ -379,7 +379,7 @@ static void SetMotionBlurShaderTempl(const FRenderingCompositePassContext& Conte
 	static FGlobalBoundShaderState BoundShaderState;
 	
 
-	SetGlobalBoundShaderState(Context.RHICmdList, BoundShaderState, GFilterVertexDeclaration.VertexDeclarationRHI, *VertexShader, *PixelShader);
+	SetGlobalBoundShaderState(Context.RHICmdList, Context.GetFeatureLevel(), BoundShaderState, GFilterVertexDeclaration.VertexDeclarationRHI, *VertexShader, *PixelShader);
 
 	VertexShader->SetParameters(Context);
 	PixelShader->SetParameters(Context);
@@ -575,7 +575,7 @@ void FRCPassPostProcessMotionBlurRecombine::Process(FRenderingCompositePassConte
 	static FGlobalBoundShaderState BoundShaderState;
 	
 
-	SetGlobalBoundShaderState(Context.RHICmdList, BoundShaderState, GFilterVertexDeclaration.VertexDeclarationRHI, *VertexShader, *PixelShader);
+	SetGlobalBoundShaderState(Context.RHICmdList, Context.GetFeatureLevel(), BoundShaderState, GFilterVertexDeclaration.VertexDeclarationRHI, *VertexShader, *PixelShader);
 
 	VertexShader->SetParameters(Context);
 	PixelShader->SetParameters(Context);
diff --git a/Engine/Source/Runtime/Renderer/Private/PostProcess/PostProcessNoiseBlur.cpp b/Engine/Source/Runtime/Renderer/Private/PostProcess/PostProcessNoiseBlur.cpp
--- a/Engine/Source/Runtime/Renderer/Private/PostProcess/PostProcessNoiseBlur.cpp
+++ b/Engine/Source/Runtime/Renderer/Private/PostProcess/PostProcessNoiseBlur.cpp
@@ -102,7 +102,7 @@ void SetNoiseBlurShader(const FRenderingCompositePassContext& Context, float InR
 	static FGlobalBoundShaderState BoundShaderState;
 	
 
-	SetGlobalBoundShaderState(Context.RHICmdList, BoundShaderState, GFilterVertexDeclaration.VertexDeclarationRHI, *VertexShader, *PixelShader);
+	SetGlobalBoundShaderState(Context.RHICmdList, Context.GetFeatureLevel(), BoundShaderState, GFilterVertexDeclaration.VertexDeclarationRHI, *VertexShader, *PixelShader);
 
 	PixelShader->SetParameters(Context, InRadius);
 	VertexShader->SetParameters(Context);
diff --git a/Engine/Source/Runtime/Renderer/Private/PostProcess/PostProcessSceneColorFringe.cpp b/Engine/Source/Runtime/Renderer/Private/PostProcess/PostProcessSceneColorFringe.cpp
--- a/Engine/Source/Runtime/Renderer/Private/PostProcess/PostProcessSceneColorFringe.cpp
+++ b/Engine/Source/Runtime/Renderer/Private/PostProcess/PostProcessSceneColorFringe.cpp
@@ -172,7 +172,7 @@ void FRCPassPostProcessSceneColorFringe::Process(FRenderingCompositePassContext&
 	static FGlobalBoundShaderState BoundShaderState;
 	
 
-	SetGlobalBoundShaderState(Context.RHICmdList, BoundShaderState, GFilterVertexDeclaration.VertexDeclarationRHI, *VertexShader, *PixelShader);
+	SetGlobalBoundShaderState(Context.RHICmdList, Context.GetFeatureLevel(), BoundShaderState, GFilterVertexDeclaration.VertexDeclarationRHI, *VertexShader, *PixelShader);
 
 	PixelShader->SetParameters(Context);
 	VertexShader->SetParameters(Context);
diff --git a/Engine/Source/Runtime/Renderer/Private/PostProcess/PostProcessSelectionOutline.cpp b/Engine/Source/Runtime/Renderer/Private/PostProcess/PostProcessSelectionOutline.cpp
--- a/Engine/Source/Runtime/Renderer/Private/PostProcess/PostProcessSelectionOutline.cpp
+++ b/Engine/Source/Runtime/Renderer/Private/PostProcess/PostProcessSelectionOutline.cpp
@@ -356,7 +356,7 @@ static void SetSelectionOutlineShaderTempl(const FRenderingCompositePassContext&
 	static FGlobalBoundShaderState BoundShaderState;
 	
 
-	SetGlobalBoundShaderState(Context.RHICmdList, BoundShaderState, GFilterVertexDeclaration.VertexDeclarationRHI, *VertexShader, *PixelShader);
+	SetGlobalBoundShaderState(Context.RHICmdList, Context.GetFeatureLevel(), BoundShaderState, GFilterVertexDeclaration.VertexDeclarationRHI, *VertexShader, *PixelShader);
 
 	PixelShader->SetPS(Context);
 }
diff --git a/Engine/Source/Runtime/Renderer/Private/PostProcess/PostProcessSubsurface.cpp b/Engine/Source/Runtime/Renderer/Private/PostProcess/PostProcessSubsurface.cpp
--- a/Engine/Source/Runtime/Renderer/Private/PostProcess/PostProcessSubsurface.cpp
+++ b/Engine/Source/Runtime/Renderer/Private/PostProcess/PostProcessSubsurface.cpp
@@ -91,7 +91,7 @@ void SetSubsurfaceSetupShader(const FRenderingCompositePassContext& Context)
 	static FGlobalBoundShaderState BoundShaderState;
 	
 
-	SetGlobalBoundShaderState(Context.RHICmdList, BoundShaderState, GFilterVertexDeclaration.VertexDeclarationRHI, *VertexShader, *PixelShader);
+	SetGlobalBoundShaderState(Context.RHICmdList, Context.GetFeatureLevel(), BoundShaderState, GFilterVertexDeclaration.VertexDeclarationRHI, *VertexShader, *PixelShader);
 
 	PixelShader->SetParameters(Context);
 	VertexShader->SetParameters(Context);
@@ -310,7 +310,7 @@ void SetSubsurfaceShader(const FRenderingCompositePassContext& Context, float In
 
 	static FGlobalBoundShaderState BoundShaderState;
 
-	SetGlobalBoundShaderState(Context.RHICmdList, BoundShaderState, GFilterVertexDeclaration.VertexDeclarationRHI, *VertexShader, *PixelShader);
+	SetGlobalBoundShaderState(Context.RHICmdList, Context.GetFeatureLevel(), BoundShaderState, GFilterVertexDeclaration.VertexDeclarationRHI, *VertexShader, *PixelShader);
 
 	PixelShader->SetParameters(Context, InRadius);
 	VertexShader->SetParameters(Context);
diff --git a/Engine/Source/Runtime/Renderer/Private/PostProcess/PostProcessTemporalAA.cpp b/Engine/Source/Runtime/Renderer/Private/PostProcess/PostProcessTemporalAA.cpp
--- a/Engine/Source/Runtime/Renderer/Private/PostProcess/PostProcessTemporalAA.cpp
+++ b/Engine/Source/Runtime/Renderer/Private/PostProcess/PostProcessTemporalAA.cpp
@@ -231,7 +231,7 @@ void FRCPassPostProcessSSRTemporalAA::Process(FRenderingCompositePassContext& Co
 	static FGlobalBoundShaderState BoundShaderState;
 	
 
-	SetGlobalBoundShaderState(Context.RHICmdList, BoundShaderState, GFilterVertexDeclaration.VertexDeclarationRHI, *VertexShader, *PixelShader);
+	SetGlobalBoundShaderState(Context.RHICmdList, Context.GetFeatureLevel(), BoundShaderState, GFilterVertexDeclaration.VertexDeclarationRHI, *VertexShader, *PixelShader);
 
 	VertexShader->SetVS(Context);
 	PixelShader->SetParameters(Context);
@@ -308,7 +308,7 @@ void FRCPassPostProcessDOFTemporalAA::Process(FRenderingCompositePassContext& Co
 	static FGlobalBoundShaderState BoundShaderState;
 	
 
-	SetGlobalBoundShaderState(Context.RHICmdList, BoundShaderState, GFilterVertexDeclaration.VertexDeclarationRHI, *VertexShader, *PixelShader);
+	SetGlobalBoundShaderState(Context.RHICmdList, Context.GetFeatureLevel(), BoundShaderState, GFilterVertexDeclaration.VertexDeclarationRHI, *VertexShader, *PixelShader);
 
 	VertexShader->SetVS(Context);
 	PixelShader->SetParameters(Context);
@@ -389,7 +389,7 @@ void FRCPassPostProcessLightShaftTemporalAA::Process(FRenderingCompositePassCont
 	static FGlobalBoundShaderState BoundShaderState;
 	
 
-	SetGlobalBoundShaderState(Context.RHICmdList, BoundShaderState, GFilterVertexDeclaration.VertexDeclarationRHI, *VertexShader, *PixelShader);
+	SetGlobalBoundShaderState(Context.RHICmdList, Context.GetFeatureLevel(), BoundShaderState, GFilterVertexDeclaration.VertexDeclarationRHI, *VertexShader, *PixelShader);
 
 	VertexShader->SetVS(Context);
 	PixelShader->SetParameters(Context);
@@ -480,7 +480,7 @@ void FRCPassPostProcessTemporalAA::Process(FRenderingCompositePassContext& Conte
 			static FGlobalBoundShaderState BoundShaderState;
 			
 
-			SetGlobalBoundShaderState(Context.RHICmdList, BoundShaderState, GFilterVertexDeclaration.VertexDeclarationRHI, *VertexShader, *PixelShader);
+			SetGlobalBoundShaderState(Context.RHICmdList, Context.GetFeatureLevel(), BoundShaderState, GFilterVertexDeclaration.VertexDeclarationRHI, *VertexShader, *PixelShader);
 
 			VertexShader->SetVS(Context);
 			PixelShader->SetParameters(Context);
@@ -492,7 +492,7 @@ void FRCPassPostProcessTemporalAA::Process(FRenderingCompositePassContext& Conte
 			static FGlobalBoundShaderState BoundShaderState;
 			
 
-			SetGlobalBoundShaderState(Context.RHICmdList, BoundShaderState, GFilterVertexDeclaration.VertexDeclarationRHI, *VertexShader, *PixelShader);
+			SetGlobalBoundShaderState(Context.RHICmdList, Context.GetFeatureLevel(), BoundShaderState, GFilterVertexDeclaration.VertexDeclarationRHI, *VertexShader, *PixelShader);
 
 			VertexShader->SetVS(Context);
 			PixelShader->SetParameters(Context);
@@ -525,7 +525,7 @@ void FRCPassPostProcessTemporalAA::Process(FRenderingCompositePassContext& Conte
 				static FGlobalBoundShaderState BoundShaderState;
 				
 	
-				SetGlobalBoundShaderState(Context.RHICmdList, BoundShaderState, GFilterVertexDeclaration.VertexDeclarationRHI, *VertexShader, *PixelShader);
+				SetGlobalBoundShaderState(Context.RHICmdList, Context.GetFeatureLevel(), BoundShaderState, GFilterVertexDeclaration.VertexDeclarationRHI, *VertexShader, *PixelShader);
 	
 				VertexShader->SetVS(Context);
 				PixelShader->SetParameters(Context);
@@ -537,7 +537,7 @@ void FRCPassPostProcessTemporalAA::Process(FRenderingCompositePassContext& Conte
 				static FGlobalBoundShaderState BoundShaderState;
 				
 	
-				SetGlobalBoundShaderState(Context.RHICmdList, BoundShaderState, GFilterVertexDeclaration.VertexDeclarationRHI, *VertexShader, *PixelShader);
+				SetGlobalBoundShaderState(Context.RHICmdList, Context.GetFeatureLevel(), BoundShaderState, GFilterVertexDeclaration.VertexDeclarationRHI, *VertexShader, *PixelShader);
 	
 				VertexShader->SetVS(Context);
 				PixelShader->SetParameters(Context);
@@ -568,7 +568,7 @@ void FRCPassPostProcessTemporalAA::Process(FRenderingCompositePassContext& Conte
 				static FGlobalBoundShaderState BoundShaderState;
 				
 	
-				SetGlobalBoundShaderState(Context.RHICmdList, BoundShaderState, GFilterVertexDeclaration.VertexDeclarationRHI, *VertexShader, *PixelShader);
+				SetGlobalBoundShaderState(Context.RHICmdList, Context.GetFeatureLevel(), BoundShaderState, GFilterVertexDeclaration.VertexDeclarationRHI, *VertexShader, *PixelShader);
 	
 				VertexShader->SetVS(Context);
 				PixelShader->SetParameters(Context);
@@ -580,7 +580,7 @@ void FRCPassPostProcessTemporalAA::Process(FRenderingCompositePassContext& Conte
 				static FGlobalBoundShaderState BoundShaderState;
 				
 	
-				SetGlobalBoundShaderState(Context.RHICmdList, BoundShaderState, GFilterVertexDeclaration.VertexDeclarationRHI, *VertexShader, *PixelShader);
+				SetGlobalBoundShaderState(Context.RHICmdList, Context.GetFeatureLevel(), BoundShaderState, GFilterVertexDeclaration.VertexDeclarationRHI, *VertexShader, *PixelShader);
 	
 				VertexShader->SetVS(Context);
 				PixelShader->SetParameters(Context);
diff --git a/Engine/Source/Runtime/Renderer/Private/PostProcess/PostProcessTestImage.cpp b/Engine/Source/Runtime/Renderer/Private/PostProcess/PostProcessTestImage.cpp
--- a/Engine/Source/Runtime/Renderer/Private/PostProcess/PostProcessTestImage.cpp
+++ b/Engine/Source/Runtime/Renderer/Private/PostProcess/PostProcessTestImage.cpp
@@ -113,7 +113,7 @@ void FRCPassPostProcessTestImage::Process(FRenderingCompositePassContext& Contex
 	static FGlobalBoundShaderState BoundShaderState;
 	
 
-	SetGlobalBoundShaderState(Context.RHICmdList, BoundShaderState, GFilterVertexDeclaration.VertexDeclarationRHI, *VertexShader, *PixelShader);
+	SetGlobalBoundShaderState(Context.RHICmdList, Context.GetFeatureLevel(), BoundShaderState, GFilterVertexDeclaration.VertexDeclarationRHI, *VertexShader, *PixelShader);
 
 	PixelShader->SetPS(Context);
 
@@ -151,7 +151,7 @@ void FRCPassPostProcessTestImage::Process(FRenderingCompositePassContext& Contex
 			}
 		} TempRenderTarget(View, (const FTexture2DRHIRef&)DestRenderTarget.TargetableTexture);
 
-		FCanvas Canvas(&TempRenderTarget, NULL, ViewFamily.CurrentRealTime, ViewFamily.CurrentWorldTime, ViewFamily.DeltaWorldTime);
+		FCanvas Canvas(&TempRenderTarget, NULL, ViewFamily.CurrentRealTime, ViewFamily.CurrentWorldTime, ViewFamily.DeltaWorldTime, Context.GetFeatureLevel());
 
 		float X = 30;
 		float Y = 8;
diff --git a/Engine/Source/Runtime/Renderer/Private/PostProcess/PostProcessTonemap.cpp b/Engine/Source/Runtime/Renderer/Private/PostProcess/PostProcessTonemap.cpp
--- a/Engine/Source/Runtime/Renderer/Private/PostProcess/PostProcessTonemap.cpp
+++ b/Engine/Source/Runtime/Renderer/Private/PostProcess/PostProcessTonemap.cpp
@@ -1189,7 +1189,7 @@ static void SetShaderTempl(const FRenderingCompositePassContext& Context)
 	static FGlobalBoundShaderState BoundShaderState;
 	
 
-	SetGlobalBoundShaderState(Context.RHICmdList, BoundShaderState, GFilterVertexDeclaration.VertexDeclarationRHI, *VertexShader, *PixelShader);
+	SetGlobalBoundShaderState(Context.RHICmdList, Context.GetFeatureLevel(), BoundShaderState, GFilterVertexDeclaration.VertexDeclarationRHI, *VertexShader, *PixelShader);
 
 	VertexShader->SetVS(Context);
 	PixelShader->SetPS(Context);
@@ -1556,7 +1556,7 @@ static void SetShaderTemplES2(const FRenderingCompositePassContext& Context, boo
 	static FGlobalBoundShaderState BoundShaderState;
 	
 
-	SetGlobalBoundShaderState(Context.RHICmdList, BoundShaderState, GFilterVertexDeclaration.VertexDeclarationRHI, *VertexShader, *PixelShader);
+	SetGlobalBoundShaderState(Context.RHICmdList, Context.GetFeatureLevel(), BoundShaderState, GFilterVertexDeclaration.VertexDeclarationRHI, *VertexShader, *PixelShader);
 
 	VertexShader->SetVS(Context);
 	PixelShader->SetPS(Context);
diff --git a/Engine/Source/Runtime/Renderer/Private/PostProcess/PostProcessUIBlur.cpp b/Engine/Source/Runtime/Renderer/Private/PostProcess/PostProcessUIBlur.cpp
--- a/Engine/Source/Runtime/Renderer/Private/PostProcess/PostProcessUIBlur.cpp
+++ b/Engine/Source/Runtime/Renderer/Private/PostProcess/PostProcessUIBlur.cpp
@@ -91,7 +91,7 @@ void FRCPassPostProcessUIBlur::Process(FRenderingCompositePassContext& Context)
 	static FGlobalBoundShaderState BoundShaderState;
 	
 
-	SetGlobalBoundShaderState(Context.RHICmdList, BoundShaderState, GFilterVertexDeclaration.VertexDeclarationRHI, *VertexShader, *PixelShader);
+	SetGlobalBoundShaderState(Context.RHICmdList, Context.GetFeatureLevel(), BoundShaderState, GFilterVertexDeclaration.VertexDeclarationRHI, *VertexShader, *PixelShader);
 
 	PixelShader->SetPS(Context);
 
diff --git a/Engine/Source/Runtime/Renderer/Private/PostProcess/PostProcessUpscale.cpp b/Engine/Source/Runtime/Renderer/Private/PostProcess/PostProcessUpscale.cpp
--- a/Engine/Source/Runtime/Renderer/Private/PostProcess/PostProcessUpscale.cpp
+++ b/Engine/Source/Runtime/Renderer/Private/PostProcess/PostProcessUpscale.cpp
@@ -120,7 +120,7 @@ void FRCPassPostProcessUpscale::SetShader(const FRenderingCompositePassContext&
 	static FGlobalBoundShaderState BoundShaderState;
 	
 
-	SetGlobalBoundShaderState(Context.RHICmdList, BoundShaderState, GFilterVertexDeclaration.VertexDeclarationRHI, *VertexShader, *PixelShader);
+	SetGlobalBoundShaderState(Context.RHICmdList, Context.GetFeatureLevel(), BoundShaderState, GFilterVertexDeclaration.VertexDeclarationRHI, *VertexShader, *PixelShader);
 
 	PixelShader->SetPS(Context);
 }
diff --git a/Engine/Source/Runtime/Renderer/Private/PostProcess/PostProcessVisualizeBuffer.cpp b/Engine/Source/Runtime/Renderer/Private/PostProcess/PostProcessVisualizeBuffer.cpp
--- a/Engine/Source/Runtime/Renderer/Private/PostProcess/PostProcessVisualizeBuffer.cpp
+++ b/Engine/Source/Runtime/Renderer/Private/PostProcess/PostProcessVisualizeBuffer.cpp
@@ -118,7 +118,7 @@ FShader* FRCPassPostProcessVisualizeBuffer::SetShaderTempl(const FRenderingCompo
 	static FGlobalBoundShaderState BoundShaderState;
 	
 
-	SetGlobalBoundShaderState(Context.RHICmdList, BoundShaderState, GFilterVertexDeclaration.VertexDeclarationRHI, *VertexShader, *PixelShader);
+	SetGlobalBoundShaderState(Context.RHICmdList, Context.GetFeatureLevel(), BoundShaderState, GFilterVertexDeclaration.VertexDeclarationRHI, *VertexShader, *PixelShader);
 
 	PixelShader->SetPS(Context);
 
@@ -177,7 +177,7 @@ void FRCPassPostProcessVisualizeBuffer::Process(FRenderingCompositePassContext&
 
 	static FGlobalBoundShaderState BoundShaderState;
 	
-	SetGlobalBoundShaderState(Context.RHICmdList, BoundShaderState, GFilterVertexDeclaration.VertexDeclarationRHI, *VertexShader, *PixelShader);
+	SetGlobalBoundShaderState(Context.RHICmdList, Context.GetFeatureLevel(), BoundShaderState, GFilterVertexDeclaration.VertexDeclarationRHI, *VertexShader, *PixelShader);
 
 	PixelShader->SetPS(Context);
 
@@ -252,7 +252,7 @@ void FRCPassPostProcessVisualizeBuffer::Process(FRenderingCompositePassContext&
 		}
 	} TempRenderTarget(View, (const FTexture2DRHIRef&)DestRenderTarget.TargetableTexture);
 
-	FCanvas Canvas(&TempRenderTarget, NULL, ViewFamily.CurrentRealTime, ViewFamily.CurrentWorldTime, ViewFamily.DeltaWorldTime);
+	FCanvas Canvas(&TempRenderTarget, NULL, ViewFamily.CurrentRealTime, ViewFamily.CurrentWorldTime, ViewFamily.DeltaWorldTime, Context.GetFeatureLevel());
 	FLinearColor LabelColor(1, 1, 0);
 	for (auto It = Labels.CreateConstIterator(); It; ++It)
 	{
diff --git a/Engine/Source/Runtime/Renderer/Private/PostProcess/PostProcessVisualizeHDR.cpp b/Engine/Source/Runtime/Renderer/Private/PostProcess/PostProcessVisualizeHDR.cpp
--- a/Engine/Source/Runtime/Renderer/Private/PostProcess/PostProcessVisualizeHDR.cpp
+++ b/Engine/Source/Runtime/Renderer/Private/PostProcess/PostProcessVisualizeHDR.cpp
@@ -168,7 +168,7 @@ void FRCPassPostProcessVisualizeHDR::Process(FRenderingCompositePassContext& Con
 	static FGlobalBoundShaderState BoundShaderState;
 	
 
-	SetGlobalBoundShaderState(Context.RHICmdList, BoundShaderState, GFilterVertexDeclaration.VertexDeclarationRHI, *VertexShader, *PixelShader);
+	SetGlobalBoundShaderState(Context.RHICmdList, Context.GetFeatureLevel(), BoundShaderState, GFilterVertexDeclaration.VertexDeclarationRHI, *VertexShader, *PixelShader);
 
 	PixelShader->SetPS(Context);
 
@@ -206,7 +206,7 @@ void FRCPassPostProcessVisualizeHDR::Process(FRenderingCompositePassContext& Con
 		}
 	} TempRenderTarget(View, (const FTexture2DRHIRef&)DestRenderTarget.TargetableTexture);
 
-	FCanvas Canvas(&TempRenderTarget, NULL, ViewFamily.CurrentRealTime, ViewFamily.CurrentWorldTime, ViewFamily.DeltaWorldTime);
+	FCanvas Canvas(&TempRenderTarget, NULL, ViewFamily.CurrentRealTime, ViewFamily.CurrentWorldTime, ViewFamily.DeltaWorldTime, Context.GetFeatureLevel());
 
 	float X = 30;
 	float Y = 8;
diff --git a/Engine/Source/Runtime/Renderer/Private/PostProcess/PostProcessWeightedSampleSum.cpp b/Engine/Source/Runtime/Renderer/Private/PostProcess/PostProcessWeightedSampleSum.cpp
--- a/Engine/Source/Runtime/Renderer/Private/PostProcess/PostProcessWeightedSampleSum.cpp
+++ b/Engine/Source/Runtime/Renderer/Private/PostProcess/PostProcessWeightedSampleSum.cpp
@@ -159,6 +159,7 @@ protected:
  */
 void SetFilterShaders(
 	FRHICommandListImmediate& RHICmdList,
+	ERHIFeatureLevel::Type FeatureLevel,
 	FSamplerStateRHIParamRef SamplerStateRHI,
 	FTextureRHIParamRef FilterTextureRHI,
 	FTextureRHIParamRef AdditiveTextureRHI,
@@ -183,7 +184,7 @@ void SetFilterShaders(
 			TShaderMapRef<TFilterPS<NumSamples, 0> > PixelShader(GetGlobalShaderMap()); \
 			{ \
 				static FGlobalBoundShaderState BoundShaderState; \
-				SetGlobalBoundShaderState(RHICmdList, BoundShaderState, GFilterVertexDeclaration.VertexDeclarationRHI, *VertexShader, *PixelShader); \
+				SetGlobalBoundShaderState(RHICmdList, FeatureLevel, BoundShaderState, GFilterVertexDeclaration.VertexDeclarationRHI, *VertexShader, *PixelShader); \
 			} \
 			PixelShader->SetParameters(RHICmdList, SamplerStateRHI, FilterTextureRHI, AdditiveTextureRHI, SampleWeights); \
 		} \
@@ -192,7 +193,7 @@ void SetFilterShaders(
 			TShaderMapRef<TFilterPS<NumSamples, 1> > PixelShader(GetGlobalShaderMap()); \
 			{ \
 				static FGlobalBoundShaderState BoundShaderState; \
-				SetGlobalBoundShaderState(RHICmdList, BoundShaderState, GFilterVertexDeclaration.VertexDeclarationRHI, *VertexShader, *PixelShader); \
+				SetGlobalBoundShaderState(RHICmdList, FeatureLevel, BoundShaderState, GFilterVertexDeclaration.VertexDeclarationRHI, *VertexShader, *PixelShader); \
 			} \
 			PixelShader->SetParameters(RHICmdList, SamplerStateRHI, FilterTextureRHI, AdditiveTextureRHI, SampleWeights); \
 		} \
@@ -201,7 +202,7 @@ void SetFilterShaders(
 			TShaderMapRef<TFilterPS<NumSamples, 2> > PixelShader(GetGlobalShaderMap()); \
 			{ \
 				static FGlobalBoundShaderState BoundShaderState; \
-				SetGlobalBoundShaderState(RHICmdList, BoundShaderState, GFilterVertexDeclaration.VertexDeclarationRHI, *VertexShader, *PixelShader); \
+				SetGlobalBoundShaderState(RHICmdList, FeatureLevel, BoundShaderState, GFilterVertexDeclaration.VertexDeclarationRHI, *VertexShader, *PixelShader); \
 			} \
 			PixelShader->SetParameters(RHICmdList, SamplerStateRHI, FilterTextureRHI, AdditiveTextureRHI, SampleWeights); \
 		} \
@@ -432,6 +433,7 @@ void FRCPassPostProcessWeightedSampleSum::Process(FRenderingCompositePassContext
 	FShader* VertexShader = nullptr;
 	SetFilterShaders(
 		Context.RHICmdList,
+		FeatureLevel,
 		TStaticSamplerState<SF_Bilinear,AM_Border,AM_Border,AM_Clamp>::GetRHI(),
 		FilterTexture,
 		AdditiveTexture,
@@ -447,7 +449,7 @@ void FRCPassPostProcessWeightedSampleSum::Process(FRenderingCompositePassContext
 	bool bRequiresClear = true;
 	// check if we have to clear the whole surface.
 	// Otherwise perform the clear when the dest rectangle has been computed.
-	if (bHasMultipleQuads || Context.View.GetFeatureLevel() == ERHIFeatureLevel::ES2)
+	if (bHasMultipleQuads || FeatureLevel == ERHIFeatureLevel::ES2)
 	{
 		Context.RHICmdList.Clear(true, FLinearColor(0, 0, 0, 0), false, 1.0f, false, 0, FIntRect());
 		bRequiresClear = false;
diff --git a/Engine/Source/Runtime/Renderer/Private/PostProcess/PostProcessing.cpp b/Engine/Source/Runtime/Renderer/Private/PostProcess/PostProcessing.cpp
--- a/Engine/Source/Runtime/Renderer/Private/PostProcess/PostProcessing.cpp
+++ b/Engine/Source/Runtime/Renderer/Private/PostProcess/PostProcessing.cpp
@@ -147,7 +147,7 @@ static FRenderingCompositeOutputRef RenderBloom(
 static FRenderingCompositeOutputRef CombineColorGradingLUTs(
 	FPostprocessContext& Context)
 {
-	FRenderingCompositePass* Pass = Context.Graph.RegisterPass(new(FMemStack::Get()) FRCPassPostProcessCombineLUTs());
+	FRenderingCompositePass* Pass = Context.Graph.RegisterPass(new(FMemStack::Get()) FRCPassPostProcessCombineLUTs(Context.View.GetShaderPlatform()));
 
 	return FRenderingCompositeOutputRef(Pass);
 }
diff --git a/Engine/Source/Runtime/Renderer/Private/PostProcess/RenderTargetPool.cpp b/Engine/Source/Runtime/Renderer/Private/PostProcess/RenderTargetPool.cpp
--- a/Engine/Source/Runtime/Renderer/Private/PostProcess/RenderTargetPool.cpp
+++ b/Engine/Source/Runtime/Renderer/Private/PostProcess/RenderTargetPool.cpp
@@ -728,7 +728,7 @@ void FRenderTargetPool::PresentContent(FRHICommandListImmediate& RHICmdList, con
 				}
 			} TempRenderTarget(View);
 
-			FCanvas Canvas(&TempRenderTarget, NULL, View.Family->CurrentRealTime, View.Family->CurrentWorldTime, View.Family->DeltaWorldTime);
+			FCanvas Canvas(&TempRenderTarget, NULL, View.Family->CurrentRealTime, View.Family->CurrentWorldTime, View.Family->DeltaWorldTime, View.GetFeatureLevel());
 
 			// TinyFont property
 			const int32 FontHeight = 12;
diff --git a/Engine/Source/Runtime/Renderer/Private/PostProcess/SceneRenderTargets.cpp b/Engine/Source/Runtime/Renderer/Private/PostProcess/SceneRenderTargets.cpp
--- a/Engine/Source/Runtime/Renderer/Private/PostProcess/SceneRenderTargets.cpp
+++ b/Engine/Source/Runtime/Renderer/Private/PostProcess/SceneRenderTargets.cpp
@@ -149,6 +149,9 @@ void FSceneRenderTargets::Allocate(const FSceneViewFamily& ViewFamily)
 {
 	check(IsInRenderingThread());
 
+	// If feature level has changed, release all previously allocated targets to the pool. If feature level has changed but
+	const auto NewFeatureLevel = ViewFamily.Scene->GetFeatureLevel();
+
 	FIntPoint DesiredBufferSize = GetSceneRenderTargetSize(ViewFamily);
 	check(DesiredBufferSize.X > 0 && DesiredBufferSize.Y > 0);
 	QuantizeBufferSize(DesiredBufferSize.X, DesiredBufferSize.Y);
@@ -178,9 +181,9 @@ void FSceneRenderTargets::Allocate(const FSceneViewFamily& ViewFamily)
 	uint32 Mobile32bpp = !IsMobileHDR() || IsMobileHDR32bpp();
 
 	static const auto CVarMobileMSAA = IConsoleManager::Get().FindTConsoleVariableDataInt(TEXT("r.MobileMSAA"));
-	int32 MobileMSAA = GRHIShaderPlatform == SP_OPENGL_ES2_IOS ? 1 : CVarMobileMSAA->GetValueOnRenderThread();
+	int32 MobileMSAA = GShaderPlatformForFeatureLevel[NewFeatureLevel] == SP_OPENGL_ES2_IOS ? 1 : CVarMobileMSAA->GetValueOnRenderThread();
 
-	bool bLightPropagationVolume = UseLightPropagationVolumeRT();
+	bool bLightPropagationVolume = UseLightPropagationVolumeRT(NewFeatureLevel);
 
 	uint32 MinShadowResolution;
 	{
@@ -189,9 +192,6 @@ void FSceneRenderTargets::Allocate(const FSceneViewFamily& ViewFamily)
 		MinShadowResolution = CVar->GetValueOnRenderThread();
 	}
 
-	// If feature level has changed, release all previously allocated targets to the pool. If feature level has changed but
-	const auto NewFeatureLevel = ViewFamily.Scene->GetFeatureLevel();
-
 	if( (BufferSize.X != DesiredBufferSize.X) ||
 		(BufferSize.Y != DesiredBufferSize.Y) ||
 		(CurrentGBufferFormat != GBufferFormat) ||
diff --git a/Engine/Source/Runtime/Renderer/Private/PostProcess/ScreenSpaceReflections.cpp b/Engine/Source/Runtime/Renderer/Private/PostProcess/ScreenSpaceReflections.cpp
--- a/Engine/Source/Runtime/Renderer/Private/PostProcess/ScreenSpaceReflections.cpp
+++ b/Engine/Source/Runtime/Renderer/Private/PostProcess/ScreenSpaceReflections.cpp
@@ -180,7 +180,7 @@ void FRCPassPostProcessScreenSpaceReflections::Process(FRenderingCompositePassCo
 	SCOPED_DRAW_EVENT(ScreenSpaceReflections, DEC_SCENE_ITEMS);
 
 	const FSceneView& View = Context.View;
-
+	const auto FeatureLevel = Context.GetFeatureLevel();
 	const FSceneRenderTargetItem& DestRenderTarget = PassOutputs[0].RequestSurface(Context);
 
 	// Set the view family's render target/viewport.
@@ -213,7 +213,7 @@ void FRCPassPostProcessScreenSpaceReflections::Process(FRenderingCompositePassCo
 		{ \
 			TShaderMapRef< FPostProcessScreenSpaceReflectionsPS<A, B> > PixelShader(GetGlobalShaderMap()); \
 			static FGlobalBoundShaderState BoundShaderState; \
-			SetGlobalBoundShaderState(Context.RHICmdList, BoundShaderState, GFilterVertexDeclaration.VertexDeclarationRHI, *VertexShader, *PixelShader); \
+			SetGlobalBoundShaderState(Context.RHICmdList, FeatureLevel, BoundShaderState, GFilterVertexDeclaration.VertexDeclarationRHI, *VertexShader, *PixelShader); \
 			VertexShader->SetParameters(Context); \
 			PixelShader->SetParameters(Context); \
 		}; \
diff --git a/Engine/Source/Runtime/Renderer/Private/PostProcess/VisualizeTexture.cpp b/Engine/Source/Runtime/Renderer/Private/PostProcess/VisualizeTexture.cpp
--- a/Engine/Source/Runtime/Renderer/Private/PostProcess/VisualizeTexture.cpp
+++ b/Engine/Source/Runtime/Renderer/Private/PostProcess/VisualizeTexture.cpp
@@ -191,7 +191,7 @@ public:
 IMPLEMENT_SHADER_TYPE(,FVisualizeTexturePresentPS,TEXT("VisualizeTexture"),TEXT("PresentPS"),SF_Pixel);
 
 
-template<uint32 TextureType> void VisualizeTextureForTextureType(FRHICommandList& RHICmdList, const FVisualizeTextureData& Data)
+template<uint32 TextureType> void VisualizeTextureForTextureType(FRHICommandList& RHICmdList, ERHIFeatureLevel::Type FeatureLevel, const FVisualizeTextureData& Data)
 {
 	TShaderMapRef<FScreenVS> VertexShader(GetGlobalShaderMap());
 	TShaderMapRef<VisualizeTexturePS<TextureType> > PixelShader(GetGlobalShaderMap());
@@ -199,7 +199,7 @@ template<uint32 TextureType> void VisualizeTextureForTextureType(FRHICommandList
 	static FGlobalBoundShaderState BoundShaderState;
 	
 
-	SetGlobalBoundShaderState(RHICmdList, BoundShaderState, GFilterVertexDeclaration.VertexDeclarationRHI, *VertexShader, *PixelShader);
+	SetGlobalBoundShaderState(RHICmdList, FeatureLevel, BoundShaderState, GFilterVertexDeclaration.VertexDeclarationRHI, *VertexShader, *PixelShader);
 	PixelShader->SetParameters(RHICmdList, Data);
 
 	DrawRectangle(
@@ -220,7 +220,7 @@ template<uint32 TextureType> void VisualizeTextureForTextureType(FRHICommandList
 		EDRF_UseTriangleOptimization);
 }
 
-void RenderVisualizeTexture(FRHICommandListImmediate& RHICmdList, const FVisualizeTextureData& Data)
+void RenderVisualizeTexture(FRHICommandListImmediate& RHICmdList, ERHIFeatureLevel::Type FeatureLevel, const FVisualizeTextureData& Data)
 {
 	if(Data.Desc.Is2DTexture())
 	{
@@ -228,30 +228,30 @@ void RenderVisualizeTexture(FRHICommandListImmediate& RHICmdList, const FVisuali
 		if(Data.Desc.NumSamples > 1)
 		{
 			// MSAA
-			VisualizeTextureForTextureType<5>(RHICmdList, Data);
+			VisualizeTextureForTextureType<5>(RHICmdList, FeatureLevel, Data);
 		}
 		else
 		{
 			// non MSAA
-			VisualizeTextureForTextureType<2>(RHICmdList, Data);
+			VisualizeTextureForTextureType<2>(RHICmdList, FeatureLevel, Data);
 		}
 	}
 	else if(Data.Desc.Is3DTexture())
 	{
 		// Volume
-		VisualizeTextureForTextureType<3>(RHICmdList, Data);
+		VisualizeTextureForTextureType<3>(RHICmdList, FeatureLevel, Data);
 	}
 	else if(Data.Desc.IsCubemap())
 	{
 		if(Data.Desc.IsArray())
 		{
 			// Cube[]
-			VisualizeTextureForTextureType<4>(RHICmdList, Data);
+			VisualizeTextureForTextureType<4>(RHICmdList, FeatureLevel, Data);
 		}
 		else
 		{
 			// Cube
-			VisualizeTextureForTextureType<0>(RHICmdList, Data);
+			VisualizeTextureForTextureType<0>(RHICmdList, FeatureLevel, Data);
 		}
 	}
 }
@@ -418,7 +418,7 @@ void FVisualizeTexture::GenerateContent(FRHICommandListImmediate& RHICmdList, co
 	{	
 		SCOPED_DRAW_EVENT(VisualizeTexture, DEC_SCENE_ITEMS);
 		// continue rendering to HDR if necessary
-		RenderVisualizeTexture(RHICmdList, VisualizeTextureData);
+		RenderVisualizeTexture(RHICmdList, FeatureLevel, VisualizeTextureData);
 	}
 
 	{
@@ -512,7 +512,7 @@ void FVisualizeTexture::PresentContent(FRHICommandListImmediate& RHICmdList, con
 	static FGlobalBoundShaderState BoundShaderState;
 	
 
-	SetGlobalBoundShaderState(RHICmdList, BoundShaderState, GFilterVertexDeclaration.VertexDeclarationRHI, *VertexShader, *PixelShader);
+	SetGlobalBoundShaderState(RHICmdList, FeatureLevel, BoundShaderState, GFilterVertexDeclaration.VertexDeclarationRHI, *VertexShader, *PixelShader);
 
 	VertexShader->SetParameters(RHICmdList, View);
 	PixelShader->SetParameters(RHICmdList, View, *VisualizeTextureContent);
@@ -555,7 +555,7 @@ void FVisualizeTexture::PresentContent(FRHICommandListImmediate& RHICmdList, con
 		}
 	} TempRenderTarget(View);
 
-	FCanvas Canvas(&TempRenderTarget, NULL, View.Family->CurrentRealTime, View.Family->CurrentWorldTime, View.Family->DeltaWorldTime);
+	FCanvas Canvas(&TempRenderTarget, NULL, View.Family->CurrentRealTime, View.Family->CurrentWorldTime, View.Family->DeltaWorldTime, View.GetFeatureLevel());
 
 	float X = 100 + View.ViewRect.Min.X;
 	float Y = 160 + View.ViewRect.Min.Y;
@@ -905,7 +905,8 @@ IPooledRenderTarget* FVisualizeTexture::GetObservedElement() const
 
 void FVisualizeTexture::OnStartFrame(const FSceneView& View)
 {
-	bEnabled = GMaxRHIFeatureLevel >= ERHIFeatureLevel::ES2;
+	FeatureLevel = View.GetFeatureLevel();
+	bEnabled = FeatureLevel >= ERHIFeatureLevel::ES2;
 	ViewRect = View.UnscaledViewRect;
 	AspectRatioConstrainedViewRect = View.Family->EngineShowFlags.CameraAspectRatioBars ? View.CameraConstrainedViewRect : ViewRect;
 
diff --git a/Engine/Source/Runtime/Renderer/Private/PostProcess/VisualizeTexture.h b/Engine/Source/Runtime/Renderer/Private/PostProcess/VisualizeTexture.h
--- a/Engine/Source/Runtime/Renderer/Private/PostProcess/VisualizeTexture.h
+++ b/Engine/Source/Runtime/Renderer/Private/PostProcess/VisualizeTexture.h
@@ -86,6 +86,9 @@ private:
 	// Flag to determine whether texture visualization is enabled, currently based on the feature level we are rendering with
 	bool bEnabled;
 
+	// Store feature level that we're currently using
+	ERHIFeatureLevel::Type FeatureLevel;
+
 	// is called by FPooledRenderTarget
 
 	void GenerateContent(FRHICommandListImmediate& RHICmdList, const FSceneRenderTargetItem& RenderTargetItem, const FPooledRenderTargetDesc& Desc);
diff --git a/Engine/Source/Runtime/Renderer/Private/ReflectionEnvironment.cpp b/Engine/Source/Runtime/Renderer/Private/ReflectionEnvironment.cpp
--- a/Engine/Source/Runtime/Renderer/Private/ReflectionEnvironment.cpp
+++ b/Engine/Source/Runtime/Renderer/Private/ReflectionEnvironment.cpp
@@ -563,7 +563,7 @@ void FDeferredShadingSceneRenderer::RenderReflectionCaptureSpecularBounceForAllV
 
 	static FGlobalBoundShaderState BoundShaderState;
 	
-	SetGlobalBoundShaderState(RHICmdList, BoundShaderState, GFilterVertexDeclaration.VertexDeclarationRHI, *VertexShader, *PixelShader);
+	SetGlobalBoundShaderState(RHICmdList, FeatureLevel, BoundShaderState, GFilterVertexDeclaration.VertexDeclarationRHI, *VertexShader, *PixelShader);
 
 	for (int32 ViewIndex = 0, Num = Views.Num(); ViewIndex < Num; ViewIndex++)
 	{
@@ -768,7 +768,7 @@ void FDeferredShadingSceneRenderer::RenderImageBasedReflectionsSM4ForAllViews(FR
 
 						static FGlobalBoundShaderState BoundShaderState;
 						
-						SetGlobalBoundShaderState(RHICmdList, BoundShaderState, GFilterVertexDeclaration.VertexDeclarationRHI, *VertexShader, *PixelShader);
+						SetGlobalBoundShaderState(RHICmdList, FeatureLevel, BoundShaderState, GFilterVertexDeclaration.VertexDeclarationRHI, *VertexShader, *PixelShader);
 
 						PixelShader->SetParameters(RHICmdList, View, ReflectionCapture);
 					}
@@ -778,7 +778,7 @@ void FDeferredShadingSceneRenderer::RenderImageBasedReflectionsSM4ForAllViews(FR
 
 						static FGlobalBoundShaderState BoundShaderState;
 						
-						SetGlobalBoundShaderState(RHICmdList, BoundShaderState, GFilterVertexDeclaration.VertexDeclarationRHI, *VertexShader, *PixelShader);
+						SetGlobalBoundShaderState(RHICmdList, FeatureLevel, BoundShaderState, GFilterVertexDeclaration.VertexDeclarationRHI, *VertexShader, *PixelShader);
 
 						PixelShader->SetParameters(RHICmdList, View, ReflectionCapture);
 					}
@@ -827,7 +827,7 @@ void FDeferredShadingSceneRenderer::RenderImageBasedReflectionsSM4ForAllViews(FR
 			{ \
 			TShaderMapRef< FReflectionApplyPS<A, B, C> > PixelShader(GetGlobalShaderMap()); \
 			static FGlobalBoundShaderState BoundShaderState; \
-			SetGlobalBoundShaderState(RHICmdList, BoundShaderState, GFilterVertexDeclaration.VertexDeclarationRHI, *VertexShader, *PixelShader); \
+			SetGlobalBoundShaderState(RHICmdList, FeatureLevel, BoundShaderState, GFilterVertexDeclaration.VertexDeclarationRHI, *VertexShader, *PixelShader); \
 			PixelShader->SetParameters(RHICmdList, View, LightAccumulation->GetRenderTargetItem().ShaderResourceTexture, SSROutput->GetRenderTargetItem().ShaderResourceTexture); \
 			}; \
 			break
diff --git a/Engine/Source/Runtime/Renderer/Private/ReflectionEnvironmentCapture.cpp b/Engine/Source/Runtime/Renderer/Private/ReflectionEnvironmentCapture.cpp
--- a/Engine/Source/Runtime/Renderer/Private/ReflectionEnvironmentCapture.cpp
+++ b/Engine/Source/Runtime/Renderer/Private/ReflectionEnvironmentCapture.cpp
@@ -274,7 +274,7 @@ private:
 IMPLEMENT_SHADER_TYPE(,FComputeBrightnessPS,TEXT("ReflectionEnvironmentShaders"),TEXT("ComputeBrightnessMain"),SF_Pixel);
 
 /** Computes the average brightness of the given reflection capture and stores it in the scene. */
-void ComputeAverageBrightness(FRHICommandList& RHICmdList)
+void ComputeAverageBrightness(FRHICommandList& RHICmdList, ERHIFeatureLevel::Type FeatureLevel)
 {
 	SetRenderTarget(RHICmdList, GSceneRenderTargets.ReflectionBrightness->GetRenderTargetItem().TargetableTexture, NULL);
 	RHICmdList.SetRasterizerState(TStaticRasterizerState<FM_Solid, CM_None>::GetRHI());
@@ -286,7 +286,7 @@ void ComputeAverageBrightness(FRHICommandList& RHICmdList)
 
 	static FGlobalBoundShaderState BoundShaderState;
 	
-	SetGlobalBoundShaderState(RHICmdList, BoundShaderState, GFilterVertexDeclaration.VertexDeclarationRHI, *VertexShader, *PixelShader);
+	SetGlobalBoundShaderState(RHICmdList, FeatureLevel, BoundShaderState, GFilterVertexDeclaration.VertexDeclarationRHI, *VertexShader, *PixelShader);
 
 	PixelShader->SetParameters(RHICmdList);
 
@@ -302,7 +302,7 @@ void ComputeAverageBrightness(FRHICommandList& RHICmdList)
 }
 
 /** Generates mips for glossiness and filters the cubemap for a given reflection. */
-void FilterReflectionEnvironment(FRHICommandListImmediate& RHICmdList, FSHVectorRGB3* OutIrradianceEnvironmentMap, bool bNormalize)
+void FilterReflectionEnvironment(FRHICommandListImmediate& RHICmdList, ERHIFeatureLevel::Type FeatureLevel, FSHVectorRGB3* OutIrradianceEnvironmentMap, bool bNormalize)
 {
 	const int32 EffectiveTopMipSize = GReflectionCaptureSize;
 	const int32 NumMips = FMath::CeilLogTwo(EffectiveTopMipSize) + 1;
@@ -340,7 +340,7 @@ void FilterReflectionEnvironment(FRHICommandListImmediate& RHICmdList, FSHVector
 				TShaderMapRef<FDownsampleGS> GeometryShader(GetGlobalShaderMap());
 				TShaderMapRef<FDownsamplePS> PixelShader(GetGlobalShaderMap());
 
-				SetGlobalBoundShaderState(RHICmdList, DownsampleBoundShaderState, GFilterVertexDeclaration.VertexDeclarationRHI, *VertexShader, *PixelShader, *GeometryShader);
+				SetGlobalBoundShaderState(RHICmdList, FeatureLevel, DownsampleBoundShaderState, GFilterVertexDeclaration.VertexDeclarationRHI, *VertexShader, *PixelShader, *GeometryShader);
 
 				// Draw each face with a geometry shader that routes to the correct render target slice
 				for (int32 CubeFace = 0; CubeFace < CubeFace_MAX; CubeFace++)
@@ -379,7 +379,7 @@ void FilterReflectionEnvironment(FRHICommandListImmediate& RHICmdList, FSHVector
 					TShaderMapRef<FScreenVSForGS> VertexShader(GetGlobalShaderMap());
 					TShaderMapRef<FDownsamplePS> PixelShader(GetGlobalShaderMap());
 
-					SetGlobalBoundShaderState(RHICmdList, DownsampleBoundShaderState, GFilterVertexDeclaration.VertexDeclarationRHI, *VertexShader, *PixelShader);
+					SetGlobalBoundShaderState(RHICmdList, FeatureLevel, DownsampleBoundShaderState, GFilterVertexDeclaration.VertexDeclarationRHI, *VertexShader, *PixelShader);
 
 					PixelShader->SetParameters(RHICmdList, CubeFace, SourceMipIndex, EffectiveSource);
 
@@ -409,10 +409,10 @@ void FilterReflectionEnvironment(FRHICommandListImmediate& RHICmdList, FSHVector
 	{
 		SCOPED_DRAW_EVENT(ComputeDiffuseIrradiance, DEC_SCENE_ITEMS);
 		check(DiffuseConvolutionSource != NULL);
-		ComputeDiffuseIrradiance(RHICmdList, DiffuseConvolutionSource->ShaderResourceTexture, DiffuseConvolutionSourceMip, OutIrradianceEnvironmentMap);
+		ComputeDiffuseIrradiance(RHICmdList, FeatureLevel, DiffuseConvolutionSource->ShaderResourceTexture, DiffuseConvolutionSourceMip, OutIrradianceEnvironmentMap);
 	}
 	
-	ComputeAverageBrightness(RHICmdList);
+	ComputeAverageBrightness(RHICmdList, FeatureLevel);
 
 	{	
 		SCOPED_DRAW_EVENT(FilterCubeMap, DEC_SCENE_ITEMS);
@@ -444,14 +444,14 @@ void FilterReflectionEnvironment(FRHICommandListImmediate& RHICmdList, FSHVector
 					PixelShader = *TShaderMapRef< TCubeFilterPS<1> >( GetGlobalShaderMap() );
 					static FGlobalBoundShaderState BoundShaderState;
 					
-					SetGlobalBoundShaderState(RHICmdList, BoundShaderState, GFilterVertexDeclaration.VertexDeclarationRHI, *VertexShader, PixelShader, *GeometryShader);
+					SetGlobalBoundShaderState(RHICmdList, FeatureLevel, BoundShaderState, GFilterVertexDeclaration.VertexDeclarationRHI, *VertexShader, PixelShader, *GeometryShader);
 				}
 				else
 				{
 					PixelShader = *TShaderMapRef< TCubeFilterPS<0> >( GetGlobalShaderMap() );
 					static FGlobalBoundShaderState BoundShaderState;
 					
-					SetGlobalBoundShaderState(RHICmdList, BoundShaderState, GFilterVertexDeclaration.VertexDeclarationRHI, *VertexShader, PixelShader, *GeometryShader);
+					SetGlobalBoundShaderState(RHICmdList, FeatureLevel, BoundShaderState, GFilterVertexDeclaration.VertexDeclarationRHI, *VertexShader, PixelShader, *GeometryShader);
 				}
 
 				for (int32 CubeFace = 0; CubeFace < CubeFace_MAX; CubeFace++)
@@ -497,14 +497,14 @@ void FilterReflectionEnvironment(FRHICommandListImmediate& RHICmdList, FSHVector
 						PixelShader = *TShaderMapRef< TCubeFilterPS<1> >( GetGlobalShaderMap() );
 						static FGlobalBoundShaderState BoundShaderState;
 						
-						SetGlobalBoundShaderState(RHICmdList, BoundShaderState, GFilterVertexDeclaration.VertexDeclarationRHI, *VertexShader, PixelShader);
+						SetGlobalBoundShaderState(RHICmdList, FeatureLevel, BoundShaderState, GFilterVertexDeclaration.VertexDeclarationRHI, *VertexShader, PixelShader);
 					}
 					else
 					{
 						PixelShader = *TShaderMapRef< TCubeFilterPS<0> >( GetGlobalShaderMap() );
 						static FGlobalBoundShaderState BoundShaderState;
 						
-						SetGlobalBoundShaderState(RHICmdList, BoundShaderState, GFilterVertexDeclaration.VertexDeclarationRHI, *VertexShader, PixelShader);
+						SetGlobalBoundShaderState(RHICmdList, FeatureLevel, BoundShaderState, GFilterVertexDeclaration.VertexDeclarationRHI, *VertexShader, PixelShader);
 					}
 
 					PixelShader->SetParameters(RHICmdList, CubeFace, MipIndex, EffectiveSource);
@@ -825,6 +825,8 @@ void CaptureSceneToScratchCubemap(FRHICommandListImmediate& RHICmdList, FSceneRe
 	FMemMark MemStackMark(FMemStack::Get());
 	// update any resources that needed a deferred update
 	FDeferredUpdateResource::UpdateResources(RHICmdList);
+
+	const auto FeatureLevel = SceneRenderer->FeatureLevel;
 	
 	{
 		SCOPED_DRAW_EVENT(CubeMapCapture, DEC_SCENE_ITEMS);
@@ -860,14 +862,14 @@ void CaptureSceneToScratchCubemap(FRHICommandListImmediate& RHICmdList, FSceneRe
 			TShaderMapRef<FRouteToCubeFaceGS> GeometryShader(GetGlobalShaderMap());
 
 			TShaderMapRef<FCopySceneColorToCubeFacePS> PixelShader(GetGlobalShaderMap());
-				SetGlobalBoundShaderState(RHICmdList, CopyColorCubemapBoundShaderState, GFilterVertexDeclaration.VertexDeclarationRHI, *VertexShader, *PixelShader, *GeometryShader);
+			SetGlobalBoundShaderState(RHICmdList, FeatureLevel, CopyColorCubemapBoundShaderState, GFilterVertexDeclaration.VertexDeclarationRHI, *VertexShader, *PixelShader, *GeometryShader);
+			
+			PixelShader->SetParameters(RHICmdList, SceneRenderer->Views[0], bCapturingForSkyLight, bLowerHemisphereIsBlack);
+			VertexShader->SetParameters(RHICmdList, SceneRenderer->Views[0]);
+			GeometryShader->SetParameters(RHICmdList, (uint32)CubeFace);
 			
-				PixelShader->SetParameters(RHICmdList, SceneRenderer->Views[0], bCapturingForSkyLight, bLowerHemisphereIsBlack);
-				VertexShader->SetParameters(RHICmdList, SceneRenderer->Views[0]);
-				GeometryShader->SetParameters(RHICmdList, (uint32)CubeFace);
-
 			DrawRectangle( 
-					RHICmdList,
+				RHICmdList,
 				ViewRect.Min.X, ViewRect.Min.Y, 
 				ViewRect.Width(), ViewRect.Height(),
 				ViewRect.Min.X, ViewRect.Min.Y, 
@@ -892,7 +894,7 @@ void CaptureSceneToScratchCubemap(FRHICommandListImmediate& RHICmdList, FSceneRe
 
 				TShaderMapRef<FCopyToCubeFaceVS> VertexShader(GetGlobalShaderMap());
 				TShaderMapRef<FCopySceneColorToCubeFacePS> PixelShader(GetGlobalShaderMap());
-				SetGlobalBoundShaderState(RHICmdList, CopyColorCubemapBoundShaderState, GFilterVertexDeclaration.VertexDeclarationRHI, *VertexShader, *PixelShader);
+				SetGlobalBoundShaderState(RHICmdList, FeatureLevel, CopyColorCubemapBoundShaderState, GFilterVertexDeclaration.VertexDeclarationRHI, *VertexShader, *PixelShader);
 
 				PixelShader->SetParameters(RHICmdList, SceneRenderer->Views[0], bCapturingForSkyLight, bLowerHemisphereIsBlack);
 				VertexShader->SetParameters(RHICmdList, SceneRenderer->Views[0]);
@@ -915,7 +917,7 @@ void CaptureSceneToScratchCubemap(FRHICommandListImmediate& RHICmdList, FSceneRe
 	delete SceneRenderer;
 }
 
-void CopyCubemapToScratchCubemap(FRHICommandList& RHICmdList, UTextureCube* SourceCubemap, bool bIsSkyLight, bool bLowerHemisphereIsBlack)
+void CopyCubemapToScratchCubemap(FRHICommandList& RHICmdList, ERHIFeatureLevel::Type FeatureLevel, UTextureCube* SourceCubemap, bool bIsSkyLight, bool bLowerHemisphereIsBlack)
 {
 	check(SourceCubemap);
 	
@@ -943,7 +945,7 @@ void CopyCubemapToScratchCubemap(FRHICommandList& RHICmdList, UTextureCube* Sour
 
 		for (uint32 CubeFace = 0; CubeFace < CubeFace_MAX; CubeFace++)
 		{
-			SetGlobalBoundShaderState(RHICmdList, CopyFromCubemapToCubemapBoundShaderState, GFilterVertexDeclaration.VertexDeclarationRHI, *VertexShader, *PixelShader, *GeometryShader);
+			SetGlobalBoundShaderState(RHICmdList, FeatureLevel, CopyFromCubemapToCubemapBoundShaderState, GFilterVertexDeclaration.VertexDeclarationRHI, *VertexShader, *PixelShader, *GeometryShader);
 			GeometryShader->SetParameters(RHICmdList, CubeFace);
 			PixelShader->SetParameters(RHICmdList, SourceCubemapResource, CubeFace, bIsSkyLight, bLowerHemisphereIsBlack);
 
@@ -978,7 +980,7 @@ void CopyCubemapToScratchCubemap(FRHICommandList& RHICmdList, UTextureCube* Sour
 			TShaderMapRef<FScreenVSForGS> VertexShader(GetGlobalShaderMap());
 			TShaderMapRef<FCopyCubemapToCubeFacePS> PixelShader(GetGlobalShaderMap());
 
-			SetGlobalBoundShaderState(RHICmdList, CopyFromCubemapToCubemapBoundShaderState, GFilterVertexDeclaration.VertexDeclarationRHI, *VertexShader, *PixelShader);
+			SetGlobalBoundShaderState(RHICmdList, FeatureLevel, CopyFromCubemapToCubemapBoundShaderState, GFilterVertexDeclaration.VertexDeclarationRHI, *VertexShader, *PixelShader);
 			PixelShader->SetParameters(RHICmdList, SourceCubemapResource, CubeFace, bIsSkyLight, bLowerHemisphereIsBlack);
 
 			DrawRectangle( 
@@ -1487,11 +1489,12 @@ void FScene::UpdateReflectionCaptureContents(UReflectionCaptureComponent* Captur
 		{
 			CaptureSceneIntoScratchCubemap(this, CaptureComponent->GetComponentLocation(), false, true, 0, false, false);
 
-			ENQUEUE_UNIQUE_RENDER_COMMAND( 
+			ENQUEUE_UNIQUE_RENDER_COMMAND_ONEPARAMETER( 
 				FilterCommand,
+				ERHIFeatureLevel::Type, FeatureLevel, GetFeatureLevel(),
 			{
 				bool bNormalize = true;
-				FilterReflectionEnvironment(RHICmdList, NULL, bNormalize);
+				FilterReflectionEnvironment(RHICmdList, FeatureLevel, NULL, bNormalize);
 			});
 
 			// Create a proxy to represent the reflection capture to the rendering thread
@@ -1557,13 +1560,14 @@ void FScene::UpdateSkyCaptureContents(const USkyLightComponent* CaptureComponent
 		}
 		else if (CaptureComponent->SourceType == SLS_SpecifiedCubemap)
 		{
-			ENQUEUE_UNIQUE_RENDER_COMMAND_TWOPARAMETER( 
+			ENQUEUE_UNIQUE_RENDER_COMMAND_THREEPARAMETER( 
 				CopyCubemapCommand,
 				UTextureCube*, SourceTexture, CaptureComponent->Cubemap,
 				bool, bLowerHemisphereIsBlack, CaptureComponent->bLowerHemisphereIsBlack,
+				ERHIFeatureLevel::Type, FeatureLevel, GetFeatureLevel(),
 			{
 				
-				CopyCubemapToScratchCubemap(RHICmdList, SourceTexture, true, bLowerHemisphereIsBlack);
+				CopyCubemapToScratchCubemap(RHICmdList, FeatureLevel, SourceTexture, true, bLowerHemisphereIsBlack);
 			});
 		}
 		else
@@ -1571,13 +1575,14 @@ void FScene::UpdateSkyCaptureContents(const USkyLightComponent* CaptureComponent
 			check(0);
 		}
 
-		ENQUEUE_UNIQUE_RENDER_COMMAND_ONEPARAMETER( 
+		ENQUEUE_UNIQUE_RENDER_COMMAND_TWOPARAMETER( 
 			FilterCommand,
 			FSHVectorRGB3*, IrradianceEnvironmentMap, &OutIrradianceEnvironmentMap,
+			ERHIFeatureLevel::Type, FeatureLevel, GetFeatureLevel(),
 		{
 			// Skylight is normalized manually in the shader
 			bool bNormalize = false;
-			FilterReflectionEnvironment(RHICmdList, IrradianceEnvironmentMap, bNormalize);
+			FilterReflectionEnvironment(RHICmdList, FeatureLevel, IrradianceEnvironmentMap, bNormalize);
 		});
 
 		// Wait until the SH coefficients have been written out by the RT before returning
diff --git a/Engine/Source/Runtime/Renderer/Private/ReflectionEnvironmentCapture.h b/Engine/Source/Runtime/Renderer/Private/ReflectionEnvironmentCapture.h
--- a/Engine/Source/Runtime/Renderer/Private/ReflectionEnvironmentCapture.h
+++ b/Engine/Source/Runtime/Renderer/Private/ReflectionEnvironmentCapture.h
@@ -8,7 +8,7 @@
 
 extern ENGINE_API int32 GReflectionCaptureSize;
 
-extern void ComputeDiffuseIrradiance(FRHICommandListImmediate& RHICmdList, FTextureRHIRef LightingSource, int32 LightingSourceMipIndex, FSHVectorRGB3* OutIrradianceEnvironmentMap);
+extern void ComputeDiffuseIrradiance(FRHICommandListImmediate& RHICmdList, ERHIFeatureLevel::Type FeatureLevel, FTextureRHIRef LightingSource, int32 LightingSourceMipIndex, FSHVectorRGB3* OutIrradianceEnvironmentMap);
 
 class FDownsampleGS : public FGlobalShader
 {
diff --git a/Engine/Source/Runtime/Renderer/Private/ReflectionEnvironmentDiffuseIrradiance.cpp b/Engine/Source/Runtime/Renderer/Private/ReflectionEnvironmentDiffuseIrradiance.cpp
--- a/Engine/Source/Runtime/Renderer/Private/ReflectionEnvironmentDiffuseIrradiance.cpp
+++ b/Engine/Source/Runtime/Renderer/Private/ReflectionEnvironmentDiffuseIrradiance.cpp
@@ -235,7 +235,7 @@ IMPLEMENT_SHADER_TYPE(,FAccumulateCubeFacesPS,TEXT("ReflectionEnvironmentShaders
 
 FGlobalBoundShaderState AccumulateCubeFacesBoundShaderState;
 
-void ComputeDiffuseIrradiance(FRHICommandListImmediate& RHICmdList, FTextureRHIRef LightingSource, int32 LightingSourceMipIndex, FSHVectorRGB3* OutIrradianceEnvironmentMap)
+void ComputeDiffuseIrradiance(FRHICommandListImmediate& RHICmdList, ERHIFeatureLevel::Type FeatureLevel, FTextureRHIRef LightingSource, int32 LightingSourceMipIndex, FSHVectorRGB3* OutIrradianceEnvironmentMap)
 {
 	for (int32 CoefficientIndex = 0; CoefficientIndex < FSHVector3::MaxSHBasis; CoefficientIndex++)
 	{
@@ -258,7 +258,7 @@ void ComputeDiffuseIrradiance(FRHICommandListImmediate& RHICmdList, FTextureRHIR
 				TShaderMapRef<FDownsampleGS> GeometryShader(GetGlobalShaderMap());
 				TShaderMapRef<FCopyDiffuseIrradiancePS> PixelShader(GetGlobalShaderMap());
 
-				SetGlobalBoundShaderState(RHICmdList, CopyDiffuseIrradianceShaderState, GFilterVertexDeclaration.VertexDeclarationRHI, *VertexShader, *PixelShader, *GeometryShader);
+				SetGlobalBoundShaderState(RHICmdList, FeatureLevel, CopyDiffuseIrradianceShaderState, GFilterVertexDeclaration.VertexDeclarationRHI, *VertexShader, *PixelShader, *GeometryShader);
 
 				// Draw each face with a geometry shader that routes to the correct render target slice
 				for (int32 CubeFace = 0; CubeFace < CubeFace_MAX; CubeFace++)
@@ -295,7 +295,7 @@ void ComputeDiffuseIrradiance(FRHICommandListImmediate& RHICmdList, FTextureRHIR
 					TShaderMapRef<FScreenVSForGS> VertexShader(GetGlobalShaderMap());
 					TShaderMapRef<FCopyDiffuseIrradiancePS> PixelShader(GetGlobalShaderMap());
 					
-					SetGlobalBoundShaderState(RHICmdList, CopyDiffuseIrradianceShaderState, GFilterVertexDeclaration.VertexDeclarationRHI, *VertexShader, *PixelShader);
+					SetGlobalBoundShaderState(RHICmdList, FeatureLevel, CopyDiffuseIrradianceShaderState, GFilterVertexDeclaration.VertexDeclarationRHI, *VertexShader, *PixelShader);
 					
 					PixelShader->SetParameters(RHICmdList, CubeFace, LightingSourceMipIndex, CoefficientIndex, MipSize, LightingSource);
 					
@@ -341,7 +341,7 @@ void ComputeDiffuseIrradiance(FRHICommandListImmediate& RHICmdList, FTextureRHIR
 					TShaderMapRef<FDownsampleGS> GeometryShader(GetGlobalShaderMap());
 					TShaderMapRef<FAccumulateDiffuseIrradiancePS> PixelShader(GetGlobalShaderMap());
 
-					SetGlobalBoundShaderState(RHICmdList, DiffuseIrradianceAccumulateShaderState, GFilterVertexDeclaration.VertexDeclarationRHI, *VertexShader, *PixelShader, *GeometryShader);
+					SetGlobalBoundShaderState(RHICmdList, FeatureLevel, DiffuseIrradianceAccumulateShaderState, GFilterVertexDeclaration.VertexDeclarationRHI, *VertexShader, *PixelShader, *GeometryShader);
 
 					// Draw each face with a geometry shader that routes to the correct render target slice
 					for (int32 CubeFace = 0; CubeFace < CubeFace_MAX; CubeFace++)
@@ -378,7 +378,7 @@ void ComputeDiffuseIrradiance(FRHICommandListImmediate& RHICmdList, FTextureRHIR
 						TShaderMapRef<FScreenVSForGS> VertexShader(GetGlobalShaderMap());
 						TShaderMapRef<FAccumulateDiffuseIrradiancePS> PixelShader(GetGlobalShaderMap());
 						
-						SetGlobalBoundShaderState(RHICmdList, DiffuseIrradianceAccumulateShaderState, GFilterVertexDeclaration.VertexDeclarationRHI, *VertexShader, *PixelShader);
+						SetGlobalBoundShaderState(RHICmdList, FeatureLevel, DiffuseIrradianceAccumulateShaderState, GFilterVertexDeclaration.VertexDeclarationRHI, *VertexShader, *PixelShader);
 						
 						PixelShader->SetParameters(RHICmdList, CubeFace, NumMips, SourceMipIndex, CoefficientIndex, EffectiveSource.ShaderResourceTexture);
 						
@@ -411,7 +411,7 @@ void ComputeDiffuseIrradiance(FRHICommandListImmediate& RHICmdList, FTextureRHIR
 
 			TShaderMapRef<FScreenVS> VertexShader(GetGlobalShaderMap());
 			TShaderMapRef<FAccumulateCubeFacesPS> PixelShader(GetGlobalShaderMap());
-			SetGlobalBoundShaderState(RHICmdList, AccumulateCubeFacesBoundShaderState, GFilterVertexDeclaration.VertexDeclarationRHI, *VertexShader, *PixelShader);
+			SetGlobalBoundShaderState(RHICmdList, FeatureLevel, AccumulateCubeFacesBoundShaderState, GFilterVertexDeclaration.VertexDeclarationRHI, *VertexShader, *PixelShader);
 
 			const int32 SourceMipIndex = NumMips - 1;
 			const int32 MipSize = 1;
diff --git a/Engine/Source/Runtime/Renderer/Private/SceneCaptureRendering.cpp b/Engine/Source/Runtime/Renderer/Private/SceneCaptureRendering.cpp
--- a/Engine/Source/Runtime/Renderer/Private/SceneCaptureRendering.cpp
+++ b/Engine/Source/Runtime/Renderer/Private/SceneCaptureRendering.cpp
@@ -48,7 +48,7 @@ void UpdateSceneCaptureContent_RenderThread(FRHICommandListImmediate& RHICmdList
 			TShaderMapRef<FScreenVS> VertexShader(GetGlobalShaderMap());
 			TShaderMapRef<FScreenPS> PixelShader(GetGlobalShaderMap());
 			static FGlobalBoundShaderState BoundShaderState;
-			SetGlobalBoundShaderState(RHICmdList, BoundShaderState, GFilterVertexDeclaration.VertexDeclarationRHI, *VertexShader, *PixelShader);
+			SetGlobalBoundShaderState(RHICmdList, SceneRenderer->FeatureLevel, BoundShaderState, GFilterVertexDeclaration.VertexDeclarationRHI, *VertexShader, *PixelShader);
 
 			FRenderingCompositePassContext Context(RHICmdList, SceneRenderer->Views[0]);
 
diff --git a/Engine/Source/Runtime/Renderer/Private/SceneHitProxyRendering.cpp b/Engine/Source/Runtime/Renderer/Private/SceneHitProxyRendering.cpp
--- a/Engine/Source/Runtime/Renderer/Private/SceneHitProxyRendering.cpp
+++ b/Engine/Source/Runtime/Renderer/Private/SceneHitProxyRendering.cpp
@@ -366,6 +366,8 @@ void RenderHitProxies(FRHICommandListImmediate& RHICmdList, const FSceneRenderer
 	auto & ViewFamily = SceneRenderer->ViewFamily;
 	auto & Views = SceneRenderer->Views;
 
+	const auto FeatureLevel = SceneRenderer->FeatureLevel;
+
 	// Dynamic vertex and index buffers need to be committed before rendering.
 	FGlobalDynamicVertexBuffer::Get().Commit();
 	FGlobalDynamicIndexBuffer::Get().Commit();
@@ -464,7 +466,7 @@ void RenderHitProxies(FRHICommandListImmediate& RHICmdList, const FSceneRenderer
 		DrawViewElements<FHitProxyDrawingPolicyFactory>(RHICmdList, View, FHitProxyDrawingPolicyFactory::ContextType(), SDPG_World, bPreFog);
 
 		// Draw the view's batched simple elements(lines, sprites, etc).
-		View.BatchedViewElements.Draw(RHICmdList, bNeedToSwitchVerticalAxis, View.ViewProjectionMatrix, View.ViewRect.Width(), View.ViewRect.Height(), true);
+		View.BatchedViewElements.Draw(RHICmdList, FeatureLevel, bNeedToSwitchVerticalAxis, View.ViewProjectionMatrix, View.ViewRect.Width(), View.ViewRect.Height(), true);
 
 		// Some elements should never be occluded (e.g. gizmos).
 		// So we render those twice, first to overwrite potentially nearer objects,
@@ -473,14 +475,14 @@ void RenderHitProxies(FRHICommandListImmediate& RHICmdList, const FSceneRenderer
 		// Draw the view's foreground elements last.
 		DrawViewElements<FHitProxyDrawingPolicyFactory>(RHICmdList, View, FHitProxyDrawingPolicyFactory::ContextType(), SDPG_Foreground, bPreFog);
 
-		View.TopBatchedViewElements.Draw(RHICmdList, bNeedToSwitchVerticalAxis, View.ViewProjectionMatrix, View.ViewRect.Width(), View.ViewRect.Height(), true);
+		View.TopBatchedViewElements.Draw(RHICmdList, FeatureLevel, bNeedToSwitchVerticalAxis, View.ViewProjectionMatrix, View.ViewRect.Width(), View.ViewRect.Height(), true);
 
 		// Note, this is a reversed Z depth surface, using CF_GreaterEqual.
 		RHICmdList.SetDepthStencilState(TStaticDepthStencilState<true, CF_GreaterEqual>::GetRHI());
 		// Draw the view's foreground elements last.
 		DrawViewElements<FHitProxyDrawingPolicyFactory>(RHICmdList, View, FHitProxyDrawingPolicyFactory::ContextType(), SDPG_Foreground, bPreFog);
 
-		View.TopBatchedViewElements.Draw(RHICmdList, bNeedToSwitchVerticalAxis, View.ViewProjectionMatrix, View.ViewRect.Width(), View.ViewRect.Height(), true);
+		View.TopBatchedViewElements.Draw(RHICmdList, FeatureLevel, bNeedToSwitchVerticalAxis, View.ViewProjectionMatrix, View.ViewRect.Width(), View.ViewRect.Height(), true);
 	}
 
 	// Finish drawing to the hit proxy render target.
@@ -540,6 +542,7 @@ void RenderHitProxies(FRHICommandListImmediate& RHICmdList, const FSceneRenderer
 	SetRenderTarget(RHICmdList, ViewFamily.RenderTarget->GetRenderTargetTexture(), FTextureRHIRef());
 	BatchedElements.Draw(
 				RHICmdList,
+				FeatureLevel,
 				bNeedToSwitchVerticalAxis,
 				PixelToView,
 				RenderTargetSize.X,
diff --git a/Engine/Source/Runtime/Renderer/Private/SceneOcclusion.cpp b/Engine/Source/Runtime/Renderer/Private/SceneOcclusion.cpp
--- a/Engine/Source/Runtime/Renderer/Private/SceneOcclusion.cpp
+++ b/Engine/Source/Runtime/Renderer/Private/SceneOcclusion.cpp
@@ -749,7 +749,7 @@ void FHZBOcclusionTester::Submit(FRHICommandListImmediate& RHICmdList, const FVi
 		TShaderMapRef< FHZBTestPS >	PixelShader( GetGlobalShaderMap() );
 
 		static FGlobalBoundShaderState BoundShaderState;
-		SetGlobalBoundShaderState(RHICmdList, BoundShaderState, GFilterVertexDeclaration.VertexDeclarationRHI, *VertexShader, *PixelShader);
+		SetGlobalBoundShaderState(RHICmdList, View.GetFeatureLevel(), BoundShaderState, GFilterVertexDeclaration.VertexDeclarationRHI, *VertexShader, *PixelShader);
 
 		PixelShader->SetParameters(RHICmdList, View,  ViewState->HZB, BoundsCenterTexture->GetRenderTargetItem().ShaderResourceTexture, BoundsExtentTexture->GetRenderTargetItem().ShaderResourceTexture );
 
@@ -894,7 +894,7 @@ void BuildHZB(FRHICommandListImmediate& RHICmdList, const FViewInfo& View)
 
 		static FGlobalBoundShaderState BoundShaderState;
 		
-		SetGlobalBoundShaderState(RHICmdList, BoundShaderState, GFilterVertexDeclaration.VertexDeclarationRHI, *VertexShader, *PixelShader);
+		SetGlobalBoundShaderState(RHICmdList, View.GetFeatureLevel(), BoundShaderState, GFilterVertexDeclaration.VertexDeclarationRHI, *VertexShader, *PixelShader);
 
 		// Imperfect sampling, doesn't matter too much
 		PixelShader->SetParameters(RHICmdList, View, GSceneRenderTargets.GetBufferSizeXY() );
@@ -928,7 +928,7 @@ void BuildHZB(FRHICommandListImmediate& RHICmdList, const FViewInfo& View)
 
 		static FGlobalBoundShaderState BoundShaderState;
 		
-		SetGlobalBoundShaderState(RHICmdList, BoundShaderState, GFilterVertexDeclaration.VertexDeclarationRHI, *VertexShader, *PixelShader);
+		SetGlobalBoundShaderState(RHICmdList, View.GetFeatureLevel(), BoundShaderState, GFilterVertexDeclaration.VertexDeclarationRHI, *VertexShader, *PixelShader);
 
 		PixelShader->SetParameters(RHICmdList, View, SrcSize, ViewState->HZB.MipSRVs[ MipIndex - 1 ] );
 
@@ -1002,7 +1002,7 @@ void FDeferredShadingSceneRenderer::BeginOcclusionTests(FRHICommandListImmediate
 
 			// Lookup the vertex shader.
 			TShaderMapRef<FOcclusionQueryVS> VertexShader(GetGlobalShaderMap());
-			SetGlobalBoundShaderState(RHICmdList, OcclusionTestBoundShaderState, GetVertexDeclarationFVector3(), *VertexShader, NULL);
+			SetGlobalBoundShaderState(RHICmdList, View.GetFeatureLevel(), OcclusionTestBoundShaderState, GetVertexDeclarationFVector3(), *VertexShader, NULL);
 			VertexShader->SetParameters(RHICmdList, View);
 
 			// Issue this frame's occlusion queries (occlusion queries from last frame may still be in flight)
diff --git a/Engine/Source/Runtime/Renderer/Private/ScenePostProcessing.cpp b/Engine/Source/Runtime/Renderer/Private/ScenePostProcessing.cpp
--- a/Engine/Source/Runtime/Renderer/Private/ScenePostProcessing.cpp
+++ b/Engine/Source/Runtime/Renderer/Private/ScenePostProcessing.cpp
@@ -124,7 +124,7 @@ void FSceneRenderer::GammaCorrectToViewportRenderTarget(FRHICommandList& RHICmdL
 
 	static FGlobalBoundShaderState PostProcessBoundShaderState;
 	
-	SetGlobalBoundShaderState(RHICmdList, PostProcessBoundShaderState, GFilterVertexDeclaration.VertexDeclarationRHI, *VertexShader, *PixelShader);
+	SetGlobalBoundShaderState(RHICmdList, View->GetFeatureLevel(), PostProcessBoundShaderState, GFilterVertexDeclaration.VertexDeclarationRHI, *VertexShader, *PixelShader);
 
 	float InvDisplayGamma = 1.0f / ViewFamily.RenderTarget->GetDisplayGamma();
 
diff --git a/Engine/Source/Runtime/Renderer/Private/ScenePrivate.h b/Engine/Source/Runtime/Renderer/Private/ScenePrivate.h
--- a/Engine/Source/Runtime/Renderer/Private/ScenePrivate.h
+++ b/Engine/Source/Runtime/Renderer/Private/ScenePrivate.h
@@ -576,7 +576,7 @@ public:
 	}
 
 	// call only if not yet created
-	void CreateLightPropagationVolumeIfNeeded();
+	void CreateLightPropagationVolumeIfNeeded(ERHIFeatureLevel::Type InFeatureLevel);
 
 	// @return can return 0 (if globally disabled)
 	FLightPropagationVolume* GetLightPropagationVolume() const { return LightPropagationVolume; }
@@ -713,7 +713,7 @@ public:
 
 		if(!(CurrentView.FinalPostProcessSettings.IndirectLightingColor * CurrentView.FinalPostProcessSettings.IndirectLightingIntensity).IsAlmostBlack())
 		{
-			CreateLightPropagationVolumeIfNeeded();
+			CreateLightPropagationVolumeIfNeeded(CurrentView.GetFeatureLevel());
 		}
 	}
 
diff --git a/Engine/Source/Runtime/Renderer/Private/SceneRendering.cpp b/Engine/Source/Runtime/Renderer/Private/SceneRendering.cpp
--- a/Engine/Source/Runtime/Renderer/Private/SceneRendering.cpp
+++ b/Engine/Source/Runtime/Renderer/Private/SceneRendering.cpp
@@ -627,9 +627,9 @@ FSceneRenderer::FSceneRenderer(const FSceneViewFamily* InViewFamily,FHitProxyCon
 	FeatureLevel = Scene->GetFeatureLevel();
 }
 
-bool FSceneRenderer::DoOcclusionQueries() const
+bool FSceneRenderer::DoOcclusionQueries(ERHIFeatureLevel::Type InFeatureLevel) const
 {
-	return !IsES2Platform(GRHIShaderPlatform)
+	return !(InFeatureLevel == ERHIFeatureLevel::ES2)
 		&& CVarAllowOcclusionQueries.GetValueOnRenderThread() != 0;
 }
 
@@ -726,7 +726,7 @@ void FSceneRenderer::RenderFinish(FRHICommandListImmediate& RHICmdList)
 				// create a temporary FCanvas object with the temp render target
 				// so it can get the screen size
 				int32 Y = 130;
-				FCanvas Canvas(&TempRenderTarget, NULL, View.Family->CurrentRealTime, View.Family->CurrentWorldTime, View.Family->DeltaWorldTime);
+				FCanvas Canvas(&TempRenderTarget, NULL, View.Family->CurrentRealTime, View.Family->CurrentWorldTime, View.Family->DeltaWorldTime, FeatureLevel);
 				if (bViewParentOrFrozen)
 				{
 					const FText StateText =
@@ -838,7 +838,7 @@ FSceneRenderer* FSceneRenderer::CreateSceneRenderer(const FSceneViewFamily* InVi
 
 void FSceneRenderer::RenderCustomDepthPass(FRHICommandListImmediate& RHICmdList)
 {
-	if(!IsFeatureLevelSupported(GRHIShaderPlatform, ERHIFeatureLevel::SM3))
+	if(FeatureLevel >= ERHIFeatureLevel::SM3)
 	{
 		// not yet supported on lower end platforms
 		return;
diff --git a/Engine/Source/Runtime/Renderer/Private/SceneRendering.h b/Engine/Source/Runtime/Renderer/Private/SceneRendering.h
--- a/Engine/Source/Runtime/Renderer/Private/SceneRendering.h
+++ b/Engine/Source/Runtime/Renderer/Private/SceneRendering.h
@@ -548,7 +548,7 @@ public:
 	/** Creates a scene renderer based on the current feature level. */
 	static FSceneRenderer* CreateSceneRenderer(const FSceneViewFamily* InViewFamily, FHitProxyConsumer* HitProxyConsumer);
 
-	bool DoOcclusionQueries() const;
+	bool DoOcclusionQueries(ERHIFeatureLevel::Type InFeatureLevel) const;
 
 protected:
 
diff --git a/Engine/Source/Runtime/Renderer/Private/SceneVisibility.cpp b/Engine/Source/Runtime/Renderer/Private/SceneVisibility.cpp
--- a/Engine/Source/Runtime/Renderer/Private/SceneVisibility.cpp
+++ b/Engine/Source/Runtime/Renderer/Private/SceneVisibility.cpp
@@ -334,7 +334,7 @@ static int32 OcclusionCull(FRHICommandListImmediate& RHICmdList, const FScene* S
 	FSceneViewState* ViewState = (FSceneViewState*)View.State;
 	
 	// Disable HZB on OpenGL platforms to avoid rendering artefacts
-	bool bHZBOcclusion = !IsOpenGLPlatform(GRHIShaderPlatform) && GHZBOcclusion;
+	bool bHZBOcclusion = !IsOpenGLPlatform(GShaderPlatformForFeatureLevel[Scene->GetFeatureLevel()]) && GHZBOcclusion;
 
 	// Use precomputed visibility data if it is available.
 	if (View.PrecomputedVisibilityData)
@@ -1042,7 +1042,7 @@ void FSceneRenderer::PreVisibilityFrameSetup(FRHICommandListImmediate& RHICmdLis
 		// HighResScreenshot should get best results so we don't do the occlusion optimization based on the former frame
 		extern bool GIsHighResScreenshot;
 		const bool bIsHitTesting = ViewFamily.EngineShowFlags.HitProxies;
-		if (GIsHighResScreenshot || !DoOcclusionQueries() || bIsHitTesting)
+		if (GIsHighResScreenshot || !DoOcclusionQueries(FeatureLevel) || bIsHitTesting)
 		{
 			View.bDisableQuerySubmissions = true;
 			View.bIgnoreExistingQueries = true;
diff --git a/Engine/Source/Runtime/Renderer/Private/ShaderComplexityRendering.cpp b/Engine/Source/Runtime/Renderer/Private/ShaderComplexityRendering.cpp
--- a/Engine/Source/Runtime/Renderer/Private/ShaderComplexityRendering.cpp
+++ b/Engine/Source/Runtime/Renderer/Private/ShaderComplexityRendering.cpp
@@ -184,7 +184,7 @@ void FRCPassPostProcessVisualizeComplexity::Process(FRenderingCompositePassConte
 
 	static FGlobalBoundShaderState ShaderComplexityBoundShaderState;
 	
-	SetGlobalBoundShaderState(Context.RHICmdList, ShaderComplexityBoundShaderState, GFilterVertexDeclaration.VertexDeclarationRHI, *VertexShader, *PixelShader);
+	SetGlobalBoundShaderState(Context.RHICmdList, Context.GetFeatureLevel(), ShaderComplexityBoundShaderState, GFilterVertexDeclaration.VertexDeclarationRHI, *VertexShader, *PixelShader);
 
 	PixelShader->SetParameters(Context, Colors);
 	
diff --git a/Engine/Source/Runtime/Renderer/Private/ShadowRendering.cpp b/Engine/Source/Runtime/Renderer/Private/ShadowRendering.cpp
--- a/Engine/Source/Runtime/Renderer/Private/ShadowRendering.cpp
+++ b/Engine/Source/Runtime/Renderer/Private/ShadowRendering.cpp
@@ -998,7 +998,7 @@ void FShadowDepthDrawingPolicyFactory::AddStaticMesh(FScene* Scene,FStaticMesh*
 		const EBlendMode BlendMode = Material->GetBlendMode();
 		const EMaterialShadingModel ShadingModel = Material->GetShadingModel();
 
-		const bool bLightPropagationVolume = UseLightPropagationVolumeRT();
+		const bool bLightPropagationVolume = UseLightPropagationVolumeRT(Scene->GetFeatureLevel());
 		const bool bTwoSided  = Material->IsTwoSided() || StaticMesh->PrimitiveSceneInfo->Proxy->CastsShadowAsTwoSided();
 		const bool bLitOpaque = !IsTranslucentBlendMode(BlendMode) && ShadingModel != MSM_Unlit;
 		if(bLightPropagationVolume && (bLitOpaque || Material->ShouldInjectEmissiveIntoLPV()))
@@ -1665,7 +1665,7 @@ static void SetShadowProjectionShaderTemplNew(FRHICommandList& RHICmdList, int32
 		// Bind shader
 		static FGlobalBoundShaderState BoundShaderState;
 		
-		SetGlobalBoundShaderState(RHICmdList, BoundShaderState, GetVertexDeclarationFVector4(), ShadowProjVS, ShadowProjPS);
+		SetGlobalBoundShaderState(RHICmdList, View.GetFeatureLevel(), BoundShaderState, GetVertexDeclarationFVector4(), ShadowProjVS, ShadowProjPS);
 
 		// Set shader parameters
 		ShadowProjVS->SetParameters(RHICmdList, View, ShadowInfo);
@@ -1684,7 +1684,7 @@ static void SetShadowProjectionShaderTemplNew(FRHICommandList& RHICmdList, int32
 
 			static FGlobalBoundShaderState BoundShaderState;
 			
-			SetGlobalBoundShaderState(RHICmdList, BoundShaderState, GetVertexDeclarationFVector4(), ShadowProjVS, ShadowProjPS);
+			SetGlobalBoundShaderState(RHICmdList, View.GetFeatureLevel(), BoundShaderState, GetVertexDeclarationFVector4(), ShadowProjVS, ShadowProjPS);
 
 			ShadowProjPS->SetParameters(RHICmdList, ViewIndex, View, ShadowInfo);
 		}
@@ -1695,7 +1695,7 @@ static void SetShadowProjectionShaderTemplNew(FRHICommandList& RHICmdList, int32
 
 			static FGlobalBoundShaderState BoundShaderState;
 			
-			SetGlobalBoundShaderState(RHICmdList, BoundShaderState, GetVertexDeclarationFVector4(), ShadowProjVS, ShadowProjPS);
+			SetGlobalBoundShaderState(RHICmdList, View.GetFeatureLevel(), BoundShaderState, GetVertexDeclarationFVector4(), ShadowProjVS, ShadowProjPS);
 
 			ShadowProjPS->SetParameters(RHICmdList, ViewIndex, View, ShadowInfo);
 		}
@@ -1713,7 +1713,7 @@ static void SetShadowProjectionShaderTemplNew(FRHICommandList& RHICmdList, int32
 
 		static FGlobalBoundShaderState BoundShaderState;
 		
-		SetGlobalBoundShaderState(RHICmdList, BoundShaderState, GetVertexDeclarationFVector4(), ShadowProjVS, ShadowProjPS);
+		SetGlobalBoundShaderState(RHICmdList, View.GetFeatureLevel(), BoundShaderState, GetVertexDeclarationFVector4(), ShadowProjVS, ShadowProjPS);
 
 		ShadowProjVS->SetParameters(RHICmdList, View, ShadowInfo);
 		ShadowProjPS->SetParameters(RHICmdList, ViewIndex, View, ShadowInfo);
@@ -1935,7 +1935,7 @@ void FProjectedShadowInfo::RenderProjection(FRHICommandListImmediate& RHICmdList
 			// Find the projection shaders.
 			TShaderMapRef<FShadowProjectionNoTransformVS> VertexShaderNoTransform(GetGlobalShaderMap());
 			VertexShaderNoTransform->SetParameters(RHICmdList, *View);
-			SetGlobalBoundShaderState(RHICmdList, MaskBoundShaderState[0], GetVertexDeclarationFVector4(), *VertexShaderNoTransform, nullptr);
+			SetGlobalBoundShaderState(RHICmdList, View->GetFeatureLevel(), MaskBoundShaderState[0], GetVertexDeclarationFVector4(), *VertexShaderNoTransform, nullptr);
 
 			FVector4 Near = View->ViewMatrices.ProjMatrix.TransformFVector4(FVector4(0, 0, CascadeSettings.SplitNear));
 			FVector4 Far = View->ViewMatrices.ProjMatrix.TransformFVector4(FVector4(0, 0, CascadeSettings.SplitFar));
@@ -2002,7 +2002,7 @@ void FProjectedShadowInfo::RenderProjection(FRHICommandListImmediate& RHICmdList
 		TShaderMapRef<FShadowProjectionVS> VertexShader(GetGlobalShaderMap());
 
 		// Cache the bound shader state
-		SetGlobalBoundShaderState(RHICmdList, MaskBoundShaderState[1], GetVertexDeclarationFVector4(), *VertexShader, NULL);
+		SetGlobalBoundShaderState(RHICmdList, View->GetFeatureLevel(), MaskBoundShaderState[1], GetVertexDeclarationFVector4(), *VertexShader, NULL);
 
 		// Set the projection vertex shader parameters
 		VertexShader->SetParameters(RHICmdList, *View, this);
@@ -2152,7 +2152,7 @@ static void SetPointLightShaderTempl(FRHICommandList& RHICmdList, int32 ViewInde
 
 	static FGlobalBoundShaderState BoundShaderState;
 	
-	SetGlobalBoundShaderState(RHICmdList, BoundShaderState, GetVertexDeclarationFVector4(), *VertexShader, *PixelShader);
+	SetGlobalBoundShaderState(RHICmdList, View.GetFeatureLevel(), BoundShaderState, GetVertexDeclarationFVector4(), *VertexShader, *PixelShader);
 
 	VertexShader->SetParameters(RHICmdList, View, ShadowInfo);
 	PixelShader->SetParameters(RHICmdList, ViewIndex, View, ShadowInfo);
diff --git a/Engine/Source/Runtime/Renderer/Private/TranslucentLighting.cpp b/Engine/Source/Runtime/Renderer/Private/TranslucentLighting.cpp
--- a/Engine/Source/Runtime/Renderer/Private/TranslucentLighting.cpp
+++ b/Engine/Source/Runtime/Renderer/Private/TranslucentLighting.cpp
@@ -539,7 +539,7 @@ void FProjectedShadowInfo::RenderTranslucencyDepths(FRHICommandList& RHICmdList,
 			);
 
 		FLinearColor ClearColors[2] = {FLinearColor(0,0,0,0), FLinearColor(0,0,0,0)};
-		DrawClearQuadMRT(RHICmdList, true, ARRAY_COUNT(ClearColors), ClearColors, false, 1.0f, false, 0);
+		DrawClearQuadMRT(RHICmdList, SceneRenderer->FeatureLevel, true, ARRAY_COUNT(ClearColors), ClearColors, false, 1.0f, false, 0);
 
 		// Set the viewport for the shadow.
 		RHICmdList.SetViewport(
@@ -1060,7 +1060,7 @@ void RasterizeToVolumeTexture(FRHICommandList& RHICmdList, FVolumeBounds VolumeB
 
 /** Helper function that clears the given volume texture render targets. */
 template<int32 NumRenderTargets>
-void ClearVolumeTextures(FRHICommandList& RHICmdList, const FTextureRHIParamRef* RenderTargets, const FLinearColor* ClearColors)
+void ClearVolumeTextures(FRHICommandList& RHICmdList, ERHIFeatureLevel::Type FeatureLevel, const FTextureRHIParamRef* RenderTargets, const FLinearColor* ClearColors)
 {
 	SetRenderTargets(RHICmdList, NumRenderTargets, RenderTargets, FTextureRHIRef(), 0, NULL);
 
@@ -1083,7 +1083,7 @@ void ClearVolumeTextures(FRHICommandList& RHICmdList, const FTextureRHIParamRef*
 		TShaderMapRef<FWriteToSliceGS> GeometryShader(GetGlobalShaderMap());
 		TShaderMapRef<TOneColorPixelShaderMRT<NumRenderTargets> > PixelShader(GetGlobalShaderMap());
 		
-		SetGlobalBoundShaderState(RHICmdList, VolumeClearBoundShaderState, GScreenVertexDeclaration.VertexDeclarationRHI, *VertexShader, *PixelShader, *GeometryShader);
+		SetGlobalBoundShaderState(RHICmdList, FeatureLevel, VolumeClearBoundShaderState, GScreenVertexDeclaration.VertexDeclarationRHI, *VertexShader, *PixelShader, *GeometryShader);
 
 		VertexShader->SetParameters(RHICmdList, VolumeBounds, GTranslucencyLightingVolumeDim);
 		GeometryShader->SetParameters(RHICmdList, VolumeBounds);
@@ -1123,7 +1123,7 @@ void FDeferredShadingSceneRenderer::ClearTranslucentVolumeLighting(FRHICommandLi
 		ClearColors[2] = FLinearColor(0, 0, 0, 0);
 		ClearColors[3] = FLinearColor(0, 0, 0, 0);
 
-		ClearVolumeTextures<ARRAY_COUNT(RenderTargets)>(RHICmdList, RenderTargets, ClearColors);
+		ClearVolumeTextures<ARRAY_COUNT(RenderTargets)>(RHICmdList, FeatureLevel, RenderTargets, ClearColors);
 	}
 }
 
@@ -1197,7 +1197,7 @@ void FDeferredShadingSceneRenderer::InjectAmbientCubemapTranslucentVolumeLightin
 
 			static FGlobalBoundShaderState BoundShaderState;
 									
-			SetGlobalBoundShaderState(RHICmdList, BoundShaderState, GScreenVertexDeclaration.VertexDeclarationRHI, *VertexShader, *PixelShader, *GeometryShader);
+			SetGlobalBoundShaderState(RHICmdList, FeatureLevel, BoundShaderState, GScreenVertexDeclaration.VertexDeclarationRHI, *VertexShader, *PixelShader, *GeometryShader);
 
 			VertexShader->SetParameters(RHICmdList, VolumeBounds, GTranslucencyLightingVolumeDim);
 			GeometryShader->SetParameters(RHICmdList, VolumeBounds);
@@ -1233,7 +1233,7 @@ void FDeferredShadingSceneRenderer::ClearTranslucentVolumePerObjectShadowing(FRH
 		ClearColors[0] = FLinearColor(1, 1, 1, 1);
 		ClearColors[1] = FLinearColor(1, 1, 1, 1);
 
-		ClearVolumeTextures<ARRAY_COUNT(RenderTargets)>(RHICmdList, RenderTargets, ClearColors);
+		ClearVolumeTextures<ARRAY_COUNT(RenderTargets)>(RHICmdList, FeatureLevel, RenderTargets, ClearColors);
 	}
 }
 
@@ -1312,7 +1312,7 @@ void FDeferredShadingSceneRenderer::AccumulateTranslucentVolumeObjectShadowing(F
 				TShaderMapRef<FWriteToSliceGS> GeometryShader(GetGlobalShaderMap());
 				TShaderMapRef<FTranslucentObjectShadowingPS> PixelShader(GetGlobalShaderMap());
 
-				SetGlobalBoundShaderState(RHICmdList, ObjectShadowingBoundShaderState, GScreenVertexDeclaration.VertexDeclarationRHI, *VertexShader, *PixelShader, *GeometryShader);
+				SetGlobalBoundShaderState(RHICmdList, FeatureLevel, ObjectShadowingBoundShaderState, GScreenVertexDeclaration.VertexDeclarationRHI, *VertexShader, *PixelShader, *GeometryShader);
 				VertexShader->SetParameters(RHICmdList, VolumeBounds, GTranslucencyLightingVolumeDim);
 				GeometryShader->SetParameters(RHICmdList, VolumeBounds);
 				PixelShader->SetParameters(RHICmdList, View, LightSceneInfo, InProjectedShadowInfo, VolumeCascadeIndex);
@@ -1730,7 +1730,7 @@ void FDeferredShadingSceneRenderer::InjectSimpleTranslucentVolumeLightingArray(F
 						TShaderMapRef<FWriteToSliceVS> VertexShader(GetGlobalShaderMap());
 						TShaderMapRef<FWriteToSliceGS> GeometryShader(GetGlobalShaderMap());
 						TShaderMapRef<FSimpleLightTranslucentLightingInjectPS> PixelShader(GetGlobalShaderMap());
-						SetGlobalBoundShaderState(RHICmdList, InjectSimpleLightBoundShaderState, GScreenVertexDeclaration.VertexDeclarationRHI, *VertexShader, *PixelShader, *GeometryShader);
+						SetGlobalBoundShaderState(RHICmdList, FeatureLevel, InjectSimpleLightBoundShaderState, GScreenVertexDeclaration.VertexDeclarationRHI, *VertexShader, *PixelShader, *GeometryShader);
 
 						VertexShader->SetParameters(RHICmdList, VolumeBounds, GTranslucencyLightingVolumeDim);
 						GeometryShader->SetParameters(RHICmdList, VolumeBounds);
@@ -1788,7 +1788,7 @@ void FDeferredShadingSceneRenderer::FilterTranslucentVolumeLighting(FRHICommandL
 				TShaderMapRef<FWriteToSliceVS> VertexShader(GetGlobalShaderMap());
 				TShaderMapRef<FWriteToSliceGS> GeometryShader(GetGlobalShaderMap());
 				TShaderMapRef<FFilterTranslucentVolumePS> PixelShader(GetGlobalShaderMap());
-				SetGlobalBoundShaderState(RHICmdList, FilterBoundShaderState, GScreenVertexDeclaration.VertexDeclarationRHI, *VertexShader, *PixelShader, *GeometryShader);
+				SetGlobalBoundShaderState(RHICmdList, FeatureLevel, FilterBoundShaderState, GScreenVertexDeclaration.VertexDeclarationRHI, *VertexShader, *PixelShader, *GeometryShader);
 
 				VertexShader->SetParameters(RHICmdList, VolumeBounds, GTranslucencyLightingVolumeDim);
 				GeometryShader->SetParameters(RHICmdList, VolumeBounds);
diff --git a/Engine/Source/Runtime/Renderer/Private/TranslucentRendering.cpp b/Engine/Source/Runtime/Renderer/Private/TranslucentRendering.cpp
--- a/Engine/Source/Runtime/Renderer/Private/TranslucentRendering.cpp
+++ b/Engine/Source/Runtime/Renderer/Private/TranslucentRendering.cpp
@@ -182,7 +182,7 @@ void FTranslucencyDrawingPolicyFactory::CopySceneColor(FRHICommandList& RHICmdLi
 
 	TShaderMapRef<FScreenVS> ScreenVertexShader(GetGlobalShaderMap());
 	TShaderMapRef<FCopySceneColorPS> PixelShader(GetGlobalShaderMap());
-	SetGlobalBoundShaderState(RHICmdList, CopySceneColorBoundShaderState, GFilterVertexDeclaration.VertexDeclarationRHI, *ScreenVertexShader, *PixelShader);
+	SetGlobalBoundShaderState(RHICmdList, View.GetFeatureLevel(), CopySceneColorBoundShaderState, GFilterVertexDeclaration.VertexDeclarationRHI, *ScreenVertexShader, *PixelShader);
 
 	/// ?
 	PixelShader->SetParameters(RHICmdList, View);
@@ -265,6 +265,7 @@ public:
 			Parameters.Mesh.VertexFactory,
 			Parameters.Mesh.MaterialRenderProxy,
 			*Parameters.Material,
+			Parameters.FeatureLevel,
 			LightMapPolicy,
 			Parameters.BlendMode,
 			// Translucent meshes need scene render targets set as textures
diff --git a/Engine/Source/Runtime/SlateRHIRenderer/Private/SlateRHIRenderer.cpp b/Engine/Source/Runtime/SlateRHIRenderer/Private/SlateRHIRenderer.cpp
--- a/Engine/Source/Runtime/SlateRHIRenderer/Private/SlateRHIRenderer.cpp
+++ b/Engine/Source/Runtime/SlateRHIRenderer/Private/SlateRHIRenderer.cpp
@@ -862,7 +862,7 @@ void FSlateRHIRenderer::CopyWindowsToVirtualScreenBuffer(const TArray<FString>&
 				TShaderMapRef<FScreenPS> PixelShader(GetGlobalShaderMap());
 
 				static FGlobalBoundShaderState BoundShaderState;
-				SetGlobalBoundShaderState(RHICmdList, BoundShaderState, Context.RendererModule->GetFilterVertexDeclaration().VertexDeclarationRHI, *VertexShader, *PixelShader);
+				SetGlobalBoundShaderState(RHICmdList, GRHIFeatureLevel, BoundShaderState, Context.RendererModule->GetFilterVertexDeclaration().VertexDeclarationRHI, *VertexShader, *PixelShader);
 
 				if( Context.WindowRect.Width() != Context.InViewportInfo->Width || Context.WindowRect.Height() != Context.InViewportInfo->Height )
 				{
diff --git a/Engine/Source/Runtime/StreamingPauseRendering/Private/StreamingPauseRendering.cpp b/Engine/Source/Runtime/StreamingPauseRendering/Private/StreamingPauseRendering.cpp
--- a/Engine/Source/Runtime/StreamingPauseRendering/Private/StreamingPauseRendering.cpp
+++ b/Engine/Source/Runtime/StreamingPauseRendering/Private/StreamingPauseRendering.cpp
@@ -77,7 +77,7 @@ void FStreamingPauseRenderingModule::BeginStreamingPause( FViewport* GameViewpor
 
 		Viewport->EnqueueBeginRenderFrame();
 
-		FCanvas Canvas( Viewport.Get(), NULL, ViewportClient->GetWorld());
+		FCanvas Canvas(Viewport.Get(), NULL, ViewportClient->GetWorld(), ViewportClient->GetWorld()->FeatureLevel);
 		{
 			ViewportClient->Draw(Viewport.Get(), &Canvas);
 		}
diff --git a/Engine/Source/Runtime/Windows/D3D11RHI/Private/D3D11Commands.cpp b/Engine/Source/Runtime/Windows/D3D11RHI/Private/D3D11Commands.cpp
--- a/Engine/Source/Runtime/Windows/D3D11RHI/Private/D3D11Commands.cpp
+++ b/Engine/Source/Runtime/Windows/D3D11RHI/Private/D3D11Commands.cpp
@@ -1942,7 +1942,7 @@ void FD3D11DynamicRHI::RHIClearMRT(bool bClearColor,int32 NumClearColors,const F
 
 		{
 			FRHICommandList_RecursiveHazardous RHICmdList;
-			SetGlobalBoundShaderState(RHICmdList, GD3D11ClearMRTBoundShaderState[FMath::Max(BoundRenderTargets.GetNumActiveTargets() - 1, 0)], GD3D11Vector4VertexDeclaration.VertexDeclarationRHI, *VertexShader, PixelShader);
+			SetGlobalBoundShaderState(RHICmdList, GMaxRHIFeatureLevel, GD3D11ClearMRTBoundShaderState[FMath::Max(BoundRenderTargets.GetNumActiveTargets() - 1, 0)], GD3D11Vector4VertexDeclaration.VertexDeclarationRHI, *VertexShader, PixelShader);
 			FLinearColor ShaderClearColors[MaxSimultaneousRenderTargets];
 			FMemory::MemZero(ShaderClearColors);
 
diff --git a/Engine/Source/Runtime/Windows/D3D11RHI/Private/D3D11Device.cpp b/Engine/Source/Runtime/Windows/D3D11RHI/Private/D3D11Device.cpp
--- a/Engine/Source/Runtime/Windows/D3D11RHI/Private/D3D11Device.cpp
+++ b/Engine/Source/Runtime/Windows/D3D11RHI/Private/D3D11Device.cpp
@@ -81,18 +81,18 @@ FD3D11DynamicRHI::FD3D11DynamicRHI(IDXGIFactory* InDXGIFactory,D3D_FEATURE_LEVEL
 	// ES2 feature level emulation in D3D11
 	if (FParse::Param(FCommandLine::Get(), TEXT("FeatureLevelES2")))
 	{
-		GMaxRHIFeatureLevel = GCurrentRHIFeatureLevel = ERHIFeatureLevel::ES2;
-		GRHIShaderPlatform = SP_PCD3D_ES2;
+		GMaxRHIFeatureLevelValue = GRHIFeatureLevelValue = ERHIFeatureLevel::ES2;
+		GRHIShaderPlatformValue = SP_PCD3D_ES2;
 	}
 	else if(FeatureLevel == D3D_FEATURE_LEVEL_11_0)
 	{
-		GMaxRHIFeatureLevel = GCurrentRHIFeatureLevel = ERHIFeatureLevel::SM5;
-		GRHIShaderPlatform = SP_PCD3D_SM5;
+		GMaxRHIFeatureLevelValue = GRHIFeatureLevelValue = ERHIFeatureLevel::SM5;
+		GRHIShaderPlatformValue = SP_PCD3D_SM5;
 	}
 	else if(FeatureLevel == D3D_FEATURE_LEVEL_10_0)
 	{
-		GMaxRHIFeatureLevel = GCurrentRHIFeatureLevel = ERHIFeatureLevel::SM4;
-		GRHIShaderPlatform = SP_PCD3D_SM4;
+		GMaxRHIFeatureLevelValue = GRHIFeatureLevelValue = ERHIFeatureLevel::SM4;
+		GRHIShaderPlatformValue = SP_PCD3D_SM4;
 	}
 
 	GPixelCenterOffset = 0.0f;	// Note that in D3D11, there is no half-texel offset (ala DX9)	
diff --git a/Engine/Source/Runtime/Windows/D3D11RHI/Private/D3D11RHI.cpp b/Engine/Source/Runtime/Windows/D3D11RHI/Private/D3D11RHI.cpp
--- a/Engine/Source/Runtime/Windows/D3D11RHI/Private/D3D11RHI.cpp
+++ b/Engine/Source/Runtime/Windows/D3D11RHI/Private/D3D11RHI.cpp
@@ -212,7 +212,7 @@ void FD3D11DynamicRHI::IssueLongGPUTask()
 			TShaderMapRef<TOneColorVS<true> > VertexShader(GetGlobalShaderMap());
 			TShaderMapRef<FLongGPUTaskPS> PixelShader(GetGlobalShaderMap());
 
-			SetGlobalBoundShaderState(RHICmdList, LongGPUTaskBoundShaderState, GD3D11Vector4VertexDeclaration.VertexDeclarationRHI, *VertexShader, *PixelShader);
+			SetGlobalBoundShaderState(RHICmdList, GMaxRHIFeatureLevel, LongGPUTaskBoundShaderState, GD3D11Vector4VertexDeclaration.VertexDeclarationRHI, *VertexShader, *PixelShader);
 
 			// Draw a fullscreen quad
 			FVector4 Vertices[4];
diff --git a/Engine/Source/Runtime/Windows/D3D11RHI/Private/D3D11RenderTarget.cpp b/Engine/Source/Runtime/Windows/D3D11RHI/Private/D3D11RenderTarget.cpp
--- a/Engine/Source/Runtime/Windows/D3D11RHI/Private/D3D11RenderTarget.cpp
+++ b/Engine/Source/Runtime/Windows/D3D11RHI/Private/D3D11RenderTarget.cpp
@@ -251,7 +251,7 @@ void FD3D11DynamicRHI::ResolveTextureUsingShader(
 	// Set the vertex and pixel shader
 	TShaderMapRef<FResolveVS> ResolveVertexShader(GetGlobalShaderMap());
 	TShaderMapRef<TPixelShader> ResolvePixelShader(GetGlobalShaderMap());
-	SetGlobalBoundShaderState(RHICmdList, ResolveBoundShaderState, GScreenVertexDeclaration.VertexDeclarationRHI, *ResolveVertexShader, *ResolvePixelShader);
+	SetGlobalBoundShaderState(RHICmdList, GMaxRHIFeatureLevel, ResolveBoundShaderState, GScreenVertexDeclaration.VertexDeclarationRHI, *ResolveVertexShader, *ResolvePixelShader);
 
 	ResolvePixelShader->SetParameters(RHICmdList, Direct3DDeviceContext,PixelShaderParameter);
 	RHICmdList.Flush(); // always call flush when using a command list in RHI implementations before doing anything else. This is super hazardous.
