Fixed alpha in erase color in Vertex Paint mode.
#ttp 334531 - Vertex Paint: Erase color defaults to full white (values = 1.0 RGB), which is the same as the paint color. Erase should be RGB 0
#branch UE4
#proj Editor.MeshPaint

diff --git a/Engine/Source/Editor/MeshPaint/Public/MeshPaintEdMode.h b/Engine/Source/Editor/MeshPaint/Public/MeshPaintEdMode.h
--- a/Engine/Source/Editor/MeshPaint/Public/MeshPaintEdMode.h
+++ b/Engine/Source/Editor/MeshPaint/Public/MeshPaintEdMode.h
@@ -103,7 +103,7 @@ public:
 		  ResourceType( EMeshPaintResource::VertexColors ),
 		  PaintMode( EMeshPaintMode::PaintColors ),
 		  PaintColor( 1.0f, 1.0f, 1.0f, 1.0f ),
-		  EraseColor( 0.0f, 0.0f, 0.0f, 1.0f ),
+		  EraseColor( 0.0f, 0.0f, 0.0f, 0.0f ),
 		  bWriteRed( true ),
 		  bWriteGreen( true ),
 		  bWriteBlue( true ),
