Slate - ESlateCheckBoxState has been renamed to ECheckBoxState.

diff --git a/Engine/Source/Editor/LevelEditor/Private/SSurfaceProperties.cpp b/Engine/Source/Editor/LevelEditor/Private/SSurfaceProperties.cpp
--- a/Engine/Source/Editor/LevelEditor/Private/SSurfaceProperties.cpp
+++ b/Engine/Source/Editor/LevelEditor/Private/SSurfaceProperties.cpp
@@ -916,14 +916,14 @@ const FSlateBrush* SSurfaceProperties::GetPreserveScaleRatioImage() const
 	return bPreserveScaleRatio ? FEditorStyle::GetBrush( TEXT("GenericLock") ) : FEditorStyle::GetBrush( TEXT("GenericUnlock") ) ;
 }
 
-ESlateCheckBoxState::Type SSurfaceProperties::IsPreserveScaleRatioChecked() const
+ECheckBoxState SSurfaceProperties::IsPreserveScaleRatioChecked() const
 {
-	return bPreserveScaleRatio ? ESlateCheckBoxState::Checked : ESlateCheckBoxState::Unchecked;
+	return bPreserveScaleRatio ? ECheckBoxState::Checked : ECheckBoxState::Unchecked;
 }
 
-void SSurfaceProperties::OnPreserveScaleRatioToggled( ESlateCheckBoxState::Type NewState )
+void SSurfaceProperties::OnPreserveScaleRatioToggled( ECheckBoxState NewState )
 {
-	bPreserveScaleRatio = (NewState == ESlateCheckBoxState::Checked) ? true : false;
+	bPreserveScaleRatio = (NewState == ECheckBoxState::Checked) ? true : false;
 	CachedScalingValueV = CachedScalingValueU;
 	GConfig->SetBool(TEXT("SurfaceSelection"), TEXT("PreserveScaleRatio"), bPreserveScaleRatio, GEditorUserSettingsIni);
 }
@@ -948,19 +948,19 @@ const FSlateBrush* SSurfaceProperties::GetTogglePanDirectionImage( TextureCoordC
 	return bUseNegativePanningV ? FEditorStyle::GetBrush( TEXT("SurfaceDetails.PanVNegative") ) : FEditorStyle::GetBrush( TEXT("SurfaceDetails.PanVPositive") ) ;
 }
 
-ESlateCheckBoxState::Type SSurfaceProperties::IsUsingNegativePanning( TextureCoordChannel Channel ) const
+ECheckBoxState SSurfaceProperties::IsUsingNegativePanning( TextureCoordChannel Channel ) const
 {
 	if (Channel == UChannel)
 	{
-		return bUseNegativePanningU ? ESlateCheckBoxState::Checked : ESlateCheckBoxState::Unchecked;
+		return bUseNegativePanningU ? ECheckBoxState::Checked : ECheckBoxState::Unchecked;
 	}
 
-	return bUseNegativePanningV ? ESlateCheckBoxState::Checked : ESlateCheckBoxState::Unchecked;
+	return bUseNegativePanningV ? ECheckBoxState::Checked : ECheckBoxState::Unchecked;
 }
 
-void SSurfaceProperties::OnTogglePanningDirection( ESlateCheckBoxState::Type NewState, TextureCoordChannel Channel )
+void SSurfaceProperties::OnTogglePanningDirection( ECheckBoxState NewState, TextureCoordChannel Channel )
 {
-	bool CheckBoxState = (NewState == ESlateCheckBoxState::Checked) ? true : false;
+	bool CheckBoxState = (NewState == ECheckBoxState::Checked) ? true : false;
 	(Channel == UChannel) ? bUseNegativePanningU = CheckBoxState : bUseNegativePanningV = CheckBoxState;
 }
 
@@ -969,14 +969,14 @@ const FSlateBrush* SSurfaceProperties::GetToggleRotationDirectionImage() const
 	return bUseNegativeRotation ? FEditorStyle::GetBrush( TEXT("SurfaceDetails.ClockwiseRotation") ) : FEditorStyle::GetBrush( TEXT("SurfaceDetails.AntiClockwiseRotation") ) ;
 }
 
-ESlateCheckBoxState::Type SSurfaceProperties::IsUsingNegativeRotation() const
+ECheckBoxState SSurfaceProperties::IsUsingNegativeRotation() const
 {
-	return bUseNegativeRotation ? ESlateCheckBoxState::Checked : ESlateCheckBoxState::Unchecked;
+	return bUseNegativeRotation ? ECheckBoxState::Checked : ECheckBoxState::Unchecked;
 }
 
-void SSurfaceProperties::OnToggleRotationDirection( ESlateCheckBoxState::Type NewState )
+void SSurfaceProperties::OnToggleRotationDirection( ECheckBoxState NewState )
 {
-	bUseNegativeRotation = (NewState == ESlateCheckBoxState::Checked) ? true : false;
+	bUseNegativeRotation = (NewState == ECheckBoxState::Checked) ? true : false;
 }
 
 #undef LOCTEXT_NAMESPACE
diff --git a/Engine/Source/Editor/LevelEditor/Private/SSurfaceProperties.h b/Engine/Source/Editor/LevelEditor/Private/SSurfaceProperties.h
--- a/Engine/Source/Editor/LevelEditor/Private/SSurfaceProperties.h
+++ b/Engine/Source/Editor/LevelEditor/Private/SSurfaceProperties.h
@@ -160,13 +160,13 @@ private:
 	 * Called by the Preserve Scaling Ratio control when clicked
 	 * @param NewState			- The new state which the control has entered into
 	 */
-	void OnPreserveScaleRatioToggled( ESlateCheckBoxState::Type NewState );
+	void OnPreserveScaleRatioToggled( ECheckBoxState NewState );
 
 	/** 
 	 * called by the Preserve scaling ratio control to determine whether or not the control is toggled
 	 * @returns The appropriate Checked state for whether the control is toggled
 	 */
-	ESlateCheckBoxState::Type IsPreserveScaleRatioChecked() const;
+	ECheckBoxState IsPreserveScaleRatioChecked() const;
 
 	/** 
 	 * Constructs the internal Table rows of the Scale combobuttons menu. 
@@ -211,14 +211,14 @@ private:
 	 * Returns the correct checkbox state for the specified texture channels panning state.
 	 * @param Channel			- Texture Coordinate Channel this control operates on.
 	 */
-	ESlateCheckBoxState::Type IsUsingNegativePanning( TextureCoordChannel Channel ) const;
+	ECheckBoxState IsUsingNegativePanning( TextureCoordChannel Channel ) const;
 
 	/** 
 	 * Toggles the specified texture channels panning state.
 	 * @param NewState			- State of the toggled control. 
 	 * @param Channel			- Texture Coordinate Channel this control operates on.
 	 */
-	void OnTogglePanningDirection( ESlateCheckBoxState::Type NewState, TextureCoordChannel Channel );
+	void OnTogglePanningDirection( ECheckBoxState NewState, TextureCoordChannel Channel );
 
 	/** 
 	 * Returns the image to display on the Toggle Rotation Direction control.
@@ -228,13 +228,13 @@ private:
 	/** 
 	 * Returns the correct checkbox state for the current rotation directions state.
 	 */
-	ESlateCheckBoxState::Type IsUsingNegativeRotation() const;
+	ECheckBoxState IsUsingNegativeRotation() const;
 
 	/** 
 	 * Toggles the Rotation directions state.
 	 * @param NewValue			- Rotation value to be applied. 
 	 */
-	void OnToggleRotationDirection( ESlateCheckBoxState::Type NewState );
+	void OnToggleRotationDirection( ECheckBoxState NewState );
 
 	/** typedef to make it easier to handle the lightmass objects */
 	typedef TArray<ULightmassPrimitiveSettingsObject*> TLightmassSettingsObjectArray;
diff --git a/Engine/Source/Editor/LevelEditor/Public/LevelEditorActions.h b/Engine/Source/Editor/LevelEditor/Public/LevelEditorActions.h
--- a/Engine/Source/Editor/LevelEditor/Public/LevelEditorActions.h
+++ b/Engine/Source/Editor/LevelEditor/Public/LevelEditorActions.h
@@ -700,10 +700,10 @@ public:
 	static void SetLightingDensityGrayscaleScale( float Value );
 	static void SetLightingDensityRenderGrayscale();
 	static bool IsLightingDensityRenderGrayscaleChecked();
-	static void SetLightingResolutionStaticMeshes( ESlateCheckBoxState::Type NewCheckedState );
-	static ESlateCheckBoxState::Type IsLightingResolutionStaticMeshesChecked();
-	static void SetLightingResolutionBSPSurfaces( ESlateCheckBoxState::Type NewCheckedState );
-	static ESlateCheckBoxState::Type IsLightingResolutionBSPSurfacesChecked();
+	static void SetLightingResolutionStaticMeshes( ECheckBoxState NewCheckedState );
+	static ECheckBoxState IsLightingResolutionStaticMeshesChecked();
+	static void SetLightingResolutionBSPSurfaces( ECheckBoxState NewCheckedState );
+	static ECheckBoxState IsLightingResolutionBSPSurfacesChecked();
 	static void SetLightingResolutionLevel( FLightmapResRatioAdjustSettings::AdjustLevels NewLevel );
 	static bool IsLightingResolutionLevelChecked( FLightmapResRatioAdjustSettings::AdjustLevels TestLevel );
 	static void SetLightingResolutionSelectedObjectsOnly();
diff --git a/Engine/Source/Editor/Matinee/Private/Matinee.cpp b/Engine/Source/Editor/Matinee/Private/Matinee.cpp
--- a/Engine/Source/Editor/Matinee/Private/Matinee.cpp
+++ b/Engine/Source/Editor/Matinee/Private/Matinee.cpp
@@ -281,24 +281,24 @@ TSharedRef<SWidget> FMatinee::AddFilterButton(UInterpFilter* Filter)
 		.OnContextMenuOpening(this, &FMatinee::CreateTabMenu);
 }
 
-void FMatinee::SetFilterActive(ESlateCheckBoxState::Type CheckStatus, UInterpFilter* Filter)
+void FMatinee::SetFilterActive(ECheckBoxState CheckStatus, UInterpFilter* Filter)
 {
-	if ( CheckStatus == ESlateCheckBoxState::Checked )
+	if ( CheckStatus == ECheckBoxState::Checked )
 	{
 		SetSelectedFilter(Filter);
 		InvalidateTrackWindowViewports();
 	}
 }
 
-ESlateCheckBoxState::Type FMatinee::GetFilterActive(UInterpFilter* Filter) const
+ECheckBoxState FMatinee::GetFilterActive(UInterpFilter* Filter) const
 {
 	if ( IData->SelectedFilter == Filter )
 	{
-		return ESlateCheckBoxState::Checked;
+		return ECheckBoxState::Checked;
 	}
 	else
 	{
-		return ESlateCheckBoxState::Unchecked;
+		return ECheckBoxState::Unchecked;
 	}
 }
 
diff --git a/Engine/Source/Editor/Matinee/Private/Matinee.h b/Engine/Source/Editor/Matinee/Private/Matinee.h
--- a/Engine/Source/Editor/Matinee/Private/Matinee.h
+++ b/Engine/Source/Editor/Matinee/Private/Matinee.h
@@ -1281,9 +1281,9 @@ protected:
 
 	TSharedRef<SWidget> AddFilterButton(UInterpFilter* Filter);
 
-	void SetFilterActive(ESlateCheckBoxState::Type CheckStatus, UInterpFilter* Filter);
+	void SetFilterActive(ECheckBoxState CheckStatus, UInterpFilter* Filter);
 
-	ESlateCheckBoxState::Type GetFilterActive(UInterpFilter* Filter) const;
+	ECheckBoxState GetFilterActive(UInterpFilter* Filter) const;
 
 	/** Holds the slate object for the Matinee Recorder. */
 	TWeakPtr<SMatineeRecorder> MatineeRecorderWindow;
diff --git a/Engine/Source/Editor/Matinee/Private/MatineeFilterButton.cpp b/Engine/Source/Editor/Matinee/Private/MatineeFilterButton.cpp
--- a/Engine/Source/Editor/Matinee/Private/MatineeFilterButton.cpp
+++ b/Engine/Source/Editor/Matinee/Private/MatineeFilterButton.cpp
@@ -28,7 +28,7 @@ FReply SMatineeFilterButton::OnMouseButtonUp( const FGeometry& MyGeometry, const
 	if ( MouseEvent.GetEffectingButton() == EKeys::RightMouseButton )
 	{
 		// Pretend as if the checkbox was clicked so that the current tab will be selected
-		OnCheckStateChanged.ExecuteIfBound( ESlateCheckBoxState::Checked );
+		OnCheckStateChanged.ExecuteIfBound( ECheckBoxState::Checked );
 
 		if ( OnContextMenuOpening.IsBound() )
 		{
diff --git a/Engine/Source/Editor/Matinee/Private/MatineeFilterButton.h b/Engine/Source/Editor/Matinee/Private/MatineeFilterButton.h
--- a/Engine/Source/Editor/Matinee/Private/MatineeFilterButton.h
+++ b/Engine/Source/Editor/Matinee/Private/MatineeFilterButton.h
@@ -8,7 +8,7 @@ public:
 
 	SLATE_BEGIN_ARGS( SMatineeFilterButton )
 		: _Text()
-		, _IsChecked( ESlateCheckBoxState::Unchecked )
+		, _IsChecked( ECheckBoxState::Unchecked )
 		, _OnCheckStateChanged()
 		, _OnContextMenuOpening()
 		{}
@@ -17,7 +17,7 @@ public:
 		SLATE_ARGUMENT( FText, Text )
 
 		/** Whether the check box is currently in a checked state */
-		SLATE_ATTRIBUTE( ESlateCheckBoxState::Type, IsChecked )
+		SLATE_ATTRIBUTE( ECheckBoxState, IsChecked )
 
 		/** Called when the checked state has changed */
 		SLATE_EVENT( FOnCheckStateChanged, OnCheckStateChanged )
diff --git a/Engine/Source/Editor/Matinee/Private/MatineeKeyReduction.cpp b/Engine/Source/Editor/Matinee/Private/MatineeKeyReduction.cpp
--- a/Engine/Source/Editor/Matinee/Private/MatineeKeyReduction.cpp
+++ b/Engine/Source/Editor/Matinee/Private/MatineeKeyReduction.cpp
@@ -44,7 +44,7 @@ public:
 private:
 
 	void SetTolerance(int32 InTolerance) { Tolerance = InTolerance; }
-	void ToggleFullInterval(ESlateCheckBoxState::Type CheckState) { FullInterval = (CheckState == ESlateCheckBoxState::Checked); }
+	void ToggleFullInterval(ECheckBoxState CheckState) { FullInterval = (CheckState == ECheckBoxState::Checked); }
 	void SetIntervalStart(float InStart) { IntervalStart = InStart; }
 	void SetIntervalEnd(float InEnd) { IntervalEnd = InEnd; }
 
diff --git a/Engine/Source/Editor/Matinee/Private/MatineeMenus.cpp b/Engine/Source/Editor/Matinee/Private/MatineeMenus.cpp
--- a/Engine/Source/Editor/Matinee/Private/MatineeMenus.cpp
+++ b/Engine/Source/Editor/Matinee/Private/MatineeMenus.cpp
@@ -4387,7 +4387,7 @@ private:
 
 	void SetIntervalStart(float InStart) { IntervalStart = InStart; }
 	void SetIntervalEnd(float InEnd) { IntervalEnd = InEnd; }
-	void ToggleFullInterval(ESlateCheckBoxState::Type CheckState) { bFullInterval = (CheckState == ESlateCheckBoxState::Checked); }
+	void ToggleFullInterval(ECheckBoxState CheckState) { bFullInterval = (CheckState == ECheckBoxState::Checked); }
 
 	bool  UseFullInterval() { return bFullInterval; }
 	TOptional<float> GetIntervalStart() const {return IntervalStart; }
diff --git a/Engine/Source/Editor/Matinee/Private/MatineeRecordMovie.cpp b/Engine/Source/Editor/Matinee/Private/MatineeRecordMovie.cpp
--- a/Engine/Source/Editor/Matinee/Private/MatineeRecordMovie.cpp
+++ b/Engine/Source/Editor/Matinee/Private/MatineeRecordMovie.cpp
@@ -122,28 +122,28 @@ public:
 	void OnFPSTextComitted( const FText& Text, ETextCommit::Type) { FPSEntry = Text; }
 	FText GetFPSText() const { return FPSEntry; }
 
-	void OnCloseEditorChecked( ESlateCheckBoxState::Type InCheckboxState ) { bCloseEditor = (InCheckboxState == ESlateCheckBoxState::Checked); }
-	ESlateCheckBoxState::Type IsCloseEditorChecked() const { return bCloseEditor ? ESlateCheckBoxState::Checked : ESlateCheckBoxState::Unchecked; }
+	void OnCloseEditorChecked( ECheckBoxState InCheckboxState ) { bCloseEditor = (InCheckboxState == ECheckBoxState::Checked); }
+	ECheckBoxState IsCloseEditorChecked() const { return bCloseEditor ? ECheckBoxState::Checked : ECheckBoxState::Unchecked; }
 
-	void OnCinematicModeChecked( ESlateCheckBoxState::Type InCheckboxState ) { bCinematicMode = (InCheckboxState == ESlateCheckBoxState::Checked); }
-	ESlateCheckBoxState::Type IsCinematicModeChecked() const { return bCinematicMode ? ESlateCheckBoxState::Checked : ESlateCheckBoxState::Unchecked; }
+	void OnCinematicModeChecked( ECheckBoxState InCheckboxState ) { bCinematicMode = (InCheckboxState == ECheckBoxState::Checked); }
+	ECheckBoxState IsCinematicModeChecked() const { return bCinematicMode ? ECheckBoxState::Checked : ECheckBoxState::Unchecked; }
 	
 	bool IsCinematicModeOptionsEnabled() const { return bCinematicMode; }
 
-	void OnDisableMovementChecked( ESlateCheckBoxState::Type InCheckboxState ) { bDisableMovement = (InCheckboxState == ESlateCheckBoxState::Checked); }
-	ESlateCheckBoxState::Type IsDisableMovementChecked() const { return bDisableMovement ? ESlateCheckBoxState::Checked : ESlateCheckBoxState::Unchecked; }
+	void OnDisableMovementChecked( ECheckBoxState InCheckboxState ) { bDisableMovement = (InCheckboxState == ECheckBoxState::Checked); }
+	ECheckBoxState IsDisableMovementChecked() const { return bDisableMovement ? ECheckBoxState::Checked : ECheckBoxState::Unchecked; }
 
-	void OnDisableTurningChecked( ESlateCheckBoxState::Type InCheckboxState ) { bDisableTurning = (InCheckboxState == ESlateCheckBoxState::Checked); }
-	ESlateCheckBoxState::Type IsDisableTurningChecked() const { return bDisableTurning ? ESlateCheckBoxState::Checked : ESlateCheckBoxState::Unchecked; }
+	void OnDisableTurningChecked( ECheckBoxState InCheckboxState ) { bDisableTurning = (InCheckboxState == ECheckBoxState::Checked); }
+	ECheckBoxState IsDisableTurningChecked() const { return bDisableTurning ? ECheckBoxState::Checked : ECheckBoxState::Unchecked; }
 
-	void OnHidePlayerChecked( ESlateCheckBoxState::Type InCheckboxState ) { bHidePlayer = (InCheckboxState == ESlateCheckBoxState::Checked); }
-	ESlateCheckBoxState::Type IsHidePlayerChecked() const { return bHidePlayer ? ESlateCheckBoxState::Checked : ESlateCheckBoxState::Unchecked; }
+	void OnHidePlayerChecked( ECheckBoxState InCheckboxState ) { bHidePlayer = (InCheckboxState == ECheckBoxState::Checked); }
+	ECheckBoxState IsHidePlayerChecked() const { return bHidePlayer ? ECheckBoxState::Checked : ECheckBoxState::Unchecked; }
 
-	void OnHideHudChecked( ESlateCheckBoxState::Type InCheckboxState ) { bHideHud = (InCheckboxState == ESlateCheckBoxState::Checked); }
-	ESlateCheckBoxState::Type IsHideHudChecked() const { return bHideHud ? ESlateCheckBoxState::Checked : ESlateCheckBoxState::Unchecked; }
+	void OnHideHudChecked( ECheckBoxState InCheckboxState ) { bHideHud = (InCheckboxState == ECheckBoxState::Checked); }
+	ECheckBoxState IsHideHudChecked() const { return bHideHud ? ECheckBoxState::Checked : ECheckBoxState::Unchecked; }
 
-	void OnDisableTextureStreamingChecked( ESlateCheckBoxState::Type InCheckboxState ) { bDisableTextureStreaming = (InCheckboxState == ESlateCheckBoxState::Checked); }
-	ESlateCheckBoxState::Type IsDisableTextureStreamingChecked() const { return bDisableTextureStreaming ? ESlateCheckBoxState::Checked : ESlateCheckBoxState::Unchecked; }
+	void OnDisableTextureStreamingChecked( ECheckBoxState InCheckboxState ) { bDisableTextureStreaming = (InCheckboxState == ECheckBoxState::Checked); }
+	ECheckBoxState IsDisableTextureStreamingChecked() const { return bDisableTextureStreaming ? ECheckBoxState::Checked : ECheckBoxState::Unchecked; }
 
 	void OnCustomWidthChanged( int32 Width, ETextCommit::Type Type );
 	void OnCustomHeightChanged( int32 Height, ETextCommit::Type Type );
diff --git a/Engine/Source/Editor/MeshPaint/Private/SMeshPaint.cpp b/Engine/Source/Editor/MeshPaint/Private/SMeshPaint.cpp
--- a/Engine/Source/Editor/MeshPaint/Private/SMeshPaint.cpp
+++ b/Engine/Source/Editor/MeshPaint/Private/SMeshPaint.cpp
@@ -235,18 +235,18 @@ public:
 
 private:
 
-	void OnImportColorChannelChanged(ESlateCheckBoxState::Type InNewValue, FImportVertexTextureHelper::ChannelsMask InColorChannelMask)
+	void OnImportColorChannelChanged(ECheckBoxState InNewValue, FImportVertexTextureHelper::ChannelsMask InColorChannelMask)
 	{
 		// Toggle the appropriate bit in the import color mask
 		ImportColorMask ^= InColorChannelMask;
 	}
 
-	ESlateCheckBoxState::Type IsRadioChecked( FImportVertexTextureHelper::ChannelsMask InColorChannelMask ) const
+	ECheckBoxState IsRadioChecked( FImportVertexTextureHelper::ChannelsMask InColorChannelMask ) const
 	{
 		// Bitwise check to see if the specified color channel should be checked
 		return (ImportColorMask & InColorChannelMask)
-			? ESlateCheckBoxState::Checked
-			: ESlateCheckBoxState::Unchecked;
+			? ECheckBoxState::Checked
+			: ECheckBoxState::Unchecked;
 	}
 
 	FReply FindTGAButtonClicked()
@@ -410,16 +410,16 @@ private:
 			];
 	}
 
-	ESlateCheckBoxState::Type IsRadioChecked( EMeshPaintResource::Type ButtonId ) const
+	ECheckBoxState IsRadioChecked( EMeshPaintResource::Type ButtonId ) const
 	{
 		return (CurrentChoice == ButtonId)
-			? ESlateCheckBoxState::Checked
-			: ESlateCheckBoxState::Unchecked;
+			? ECheckBoxState::Checked
+			: ECheckBoxState::Unchecked;
 	}
 
-	void OnRadioChanged( ESlateCheckBoxState::Type NewRadioState, EMeshPaintResource::Type RadioThatChanged )
+	void OnRadioChanged( ECheckBoxState NewRadioState, EMeshPaintResource::Type RadioThatChanged )
 	{
-		if (NewRadioState == ESlateCheckBoxState::Checked)
+		if (NewRadioState == ECheckBoxState::Checked)
 		{
 			CurrentChoice = RadioThatChanged;
 
@@ -491,16 +491,16 @@ private:
 			];
 	}
 
-	ESlateCheckBoxState::Type IsRadioChecked( EMeshVertexPaintTarget::Type ButtonId ) const
+	ECheckBoxState IsRadioChecked( EMeshVertexPaintTarget::Type ButtonId ) const
 	{
 		return (CurrentChoice == ButtonId)
-			? ESlateCheckBoxState::Checked
-			: ESlateCheckBoxState::Unchecked;
+			? ECheckBoxState::Checked
+			: ECheckBoxState::Unchecked;
 	}
 
-	void OnRadioChanged( ESlateCheckBoxState::Type NewRadioState, EMeshVertexPaintTarget::Type RadioThatChanged )
+	void OnRadioChanged( ECheckBoxState NewRadioState, EMeshVertexPaintTarget::Type RadioThatChanged )
 	{
-		if (NewRadioState == ESlateCheckBoxState::Checked)
+		if (NewRadioState == ECheckBoxState::Checked)
 		{
 			CurrentChoice = RadioThatChanged;
 
@@ -573,16 +573,16 @@ private:
 			];
 	}
 
-	ESlateCheckBoxState::Type IsRadioChecked( EMeshPaintMode::Type ButtonId ) const
+	ECheckBoxState IsRadioChecked( EMeshPaintMode::Type ButtonId ) const
 	{
 		return (CurrentChoice == ButtonId)
-			? ESlateCheckBoxState::Checked
-			: ESlateCheckBoxState::Unchecked;
+			? ECheckBoxState::Checked
+			: ECheckBoxState::Unchecked;
 	}
 
-	void OnRadioChanged( ESlateCheckBoxState::Type NewRadioState, EMeshPaintMode::Type RadioThatChanged )
+	void OnRadioChanged( ECheckBoxState NewRadioState, EMeshPaintMode::Type RadioThatChanged )
 	{
-		if (NewRadioState == ESlateCheckBoxState::Checked)
+		if (NewRadioState == ECheckBoxState::Checked)
 		{
 			CurrentChoice = RadioThatChanged;
 
@@ -680,16 +680,16 @@ private:
 			];
 	}
 
-	ESlateCheckBoxState::Type IsRadioChecked( EMeshPaintColorViewMode::Type ButtonId ) const
+	ECheckBoxState IsRadioChecked( EMeshPaintColorViewMode::Type ButtonId ) const
 	{
 		return (CurrentChoice == ButtonId)
-			? ESlateCheckBoxState::Checked
-			: ESlateCheckBoxState::Unchecked;
+			? ECheckBoxState::Checked
+			: ECheckBoxState::Unchecked;
 	}
 
-	void OnRadioChanged( ESlateCheckBoxState::Type NewRadioState, EMeshPaintColorViewMode::Type RadioThatChanged )
+	void OnRadioChanged( ECheckBoxState NewRadioState, EMeshPaintColorViewMode::Type RadioThatChanged )
 	{
-		if (NewRadioState == ESlateCheckBoxState::Checked)
+		if (NewRadioState == ECheckBoxState::Checked)
 		{
 			CurrentChoice = RadioThatChanged;
 
@@ -861,16 +861,16 @@ private:
 			];
 	}
 
-	ESlateCheckBoxState::Type IsRadioChecked( EMeshPaintColorSet::Type ButtonId ) const
+	ECheckBoxState IsRadioChecked( EMeshPaintColorSet::Type ButtonId ) const
 	{
 		return (CurrentChoice == ButtonId)
-			? ESlateCheckBoxState::Checked
-			: ESlateCheckBoxState::Unchecked;
+			? ECheckBoxState::Checked
+			: ECheckBoxState::Unchecked;
 	}
 
-	void OnRadioChanged( ESlateCheckBoxState::Type NewRadioState, EMeshPaintColorSet::Type RadioThatChanged )
+	void OnRadioChanged( ECheckBoxState NewRadioState, EMeshPaintColorSet::Type RadioThatChanged )
 	{
-		if (NewRadioState == ESlateCheckBoxState::Checked)
+		if (NewRadioState == ECheckBoxState::Checked)
 		{
 			CurrentChoice = RadioThatChanged;
 
@@ -1030,16 +1030,16 @@ private:
 			];
 	}
 
-	ESlateCheckBoxState::Type IsRadioChecked( int32 ButtonId ) const
+	ECheckBoxState IsRadioChecked( int32 ButtonId ) const
 	{
 		return (CurrentChoice == ButtonId)
-			? ESlateCheckBoxState::Checked
-			: ESlateCheckBoxState::Unchecked;
+			? ECheckBoxState::Checked
+			: ECheckBoxState::Unchecked;
 	}
 
-	void OnRadioChanged( ESlateCheckBoxState::Type NewRadioState, int32 RadioThatChanged )
+	void OnRadioChanged( ECheckBoxState NewRadioState, int32 RadioThatChanged )
 	{
-		if (NewRadioState == ESlateCheckBoxState::Checked)
+		if (NewRadioState == ECheckBoxState::Checked)
 		{
 			if (OnSelectionChanged.IsBound())
 			{
@@ -2220,40 +2220,40 @@ void SMeshPaint::OnFlowAmountChanged(float InFlowAmount)
 	FMeshPaintSettings::Get().FlowAmount = InFlowAmount;
 }
 
-ESlateCheckBoxState::Type SMeshPaint::IsIgnoreBackfaceChecked() const
+ECheckBoxState SMeshPaint::IsIgnoreBackfaceChecked() const
 {
 	return (FMeshPaintSettings::Get().bOnlyFrontFacingTriangles)
-		? ESlateCheckBoxState::Checked
-		: ESlateCheckBoxState::Unchecked;
+		? ECheckBoxState::Checked
+		: ECheckBoxState::Unchecked;
 }
 
-void SMeshPaint::OnIgnoreBackfaceChanged(ESlateCheckBoxState::Type InCheckState)
+void SMeshPaint::OnIgnoreBackfaceChanged(ECheckBoxState InCheckState)
 {
-	FMeshPaintSettings::Get().bOnlyFrontFacingTriangles = (InCheckState == ESlateCheckBoxState::Checked);
+	FMeshPaintSettings::Get().bOnlyFrontFacingTriangles = (InCheckState == ECheckBoxState::Checked);
 }
 
-ESlateCheckBoxState::Type SMeshPaint::IsSeamPaintingChecked() const
+ECheckBoxState SMeshPaint::IsSeamPaintingChecked() const
 {
 	return (FMeshPaintSettings::Get().bEnableSeamPainting)
-		? ESlateCheckBoxState::Checked
-		: ESlateCheckBoxState::Unchecked;
+		? ECheckBoxState::Checked
+		: ECheckBoxState::Unchecked;
 }
 
-void SMeshPaint::OnSeamPaintingChanged(ESlateCheckBoxState::Type InCheckState)
+void SMeshPaint::OnSeamPaintingChanged(ECheckBoxState InCheckState)
 {
-	FMeshPaintSettings::Get().bEnableSeamPainting = (InCheckState == ESlateCheckBoxState::Checked);
+	FMeshPaintSettings::Get().bEnableSeamPainting = (InCheckState == ECheckBoxState::Checked);
 }
 
-ESlateCheckBoxState::Type SMeshPaint::IsEnableFlowChecked() const
+ECheckBoxState SMeshPaint::IsEnableFlowChecked() const
 {
 	return (FMeshPaintSettings::Get().bEnableFlow)
-		? ESlateCheckBoxState::Checked
-		: ESlateCheckBoxState::Unchecked;
+		? ECheckBoxState::Checked
+		: ECheckBoxState::Unchecked;
 }
 
-void SMeshPaint::OnEnableFlowChanged(ESlateCheckBoxState::Type InCheckState)
+void SMeshPaint::OnEnableFlowChanged(ECheckBoxState InCheckState)
 {
-	FMeshPaintSettings::Get().bEnableFlow = (InCheckState == ESlateCheckBoxState::Checked);
+	FMeshPaintSettings::Get().bEnableFlow = (InCheckState == ECheckBoxState::Checked);
 }
 
 void SMeshPaint::OnEraseWeightChanged(int32 InWeightIndex)
@@ -2267,7 +2267,7 @@ void SMeshPaint::OnPaintWeightChanged(int32 InWeightIndex)
 }
 
 
-ESlateCheckBoxState::Type SMeshPaint::IsWriteColorChannelChecked(EMeshPaintWriteColorChannels::Type CheckBoxInfo) const
+ECheckBoxState SMeshPaint::IsWriteColorChannelChecked(EMeshPaintWriteColorChannels::Type CheckBoxInfo) const
 {
 	const bool bIsRedAndChecked = ( CheckBoxInfo == EMeshPaintWriteColorChannels::Red && FMeshPaintSettings::Get().bWriteRed );
 	const bool bIsGreenAndChecked = ( CheckBoxInfo == EMeshPaintWriteColorChannels::Green && FMeshPaintSettings::Get().bWriteGreen );
@@ -2275,12 +2275,12 @@ ESlateCheckBoxState::Type SMeshPaint::IsWriteColorChannelChecked(EMeshPaintWrite
 	const bool bIsAlphaAndChecked = (CheckBoxInfo == EMeshPaintWriteColorChannels::Alpha && FMeshPaintSettings::Get().bWriteAlpha);
 
 	const bool bIsChecked = ( bIsRedAndChecked || bIsGreenAndChecked || bIsBlueAndChecked || bIsAlphaAndChecked );
-	return bIsChecked ? ESlateCheckBoxState::Checked : ESlateCheckBoxState::Unchecked;
+	return bIsChecked ? ECheckBoxState::Checked : ECheckBoxState::Unchecked;
 }
 
-void SMeshPaint::OnWriteColorChannelChanged(ESlateCheckBoxState::Type InNewValue, EMeshPaintWriteColorChannels::Type CheckBoxInfo)
+void SMeshPaint::OnWriteColorChannelChanged(ECheckBoxState InNewValue, EMeshPaintWriteColorChannels::Type CheckBoxInfo)
 {
-	bool bIsCheckedState = (InNewValue == ESlateCheckBoxState::Checked);
+	bool bIsCheckedState = (InNewValue == ECheckBoxState::Checked);
 	if(CheckBoxInfo == EMeshPaintWriteColorChannels::Red)
 	{
 		FMeshPaintSettings::Get().bWriteRed = bIsCheckedState;
diff --git a/Engine/Source/Editor/MeshPaint/Private/SMeshPaint.h b/Engine/Source/Editor/MeshPaint/Private/SMeshPaint.h
--- a/Engine/Source/Editor/MeshPaint/Private/SMeshPaint.h
+++ b/Engine/Source/Editor/MeshPaint/Private/SMeshPaint.h
@@ -141,22 +141,22 @@ protected:
 	void OnFlowAmountChanged(float InFlowAmount);
 
 	/** Called to get the checked state of the ignore back-facing checkbox */
-	ESlateCheckBoxState::Type IsIgnoreBackfaceChecked() const;
+	ECheckBoxState IsIgnoreBackfaceChecked() const;
 
 	/** Called when the ignore back face checkbox value is changed */
-	void OnIgnoreBackfaceChanged(ESlateCheckBoxState::Type InCheckState);
+	void OnIgnoreBackfaceChanged(ECheckBoxState InCheckState);
 
 	/** Called to get the checked state of the seam painting checkbox */
-	ESlateCheckBoxState::Type IsSeamPaintingChecked() const;
+	ECheckBoxState IsSeamPaintingChecked() const;
 
 	/** Called when the seam painting checkbox value is changed */
-	void OnSeamPaintingChanged(ESlateCheckBoxState::Type InCheckState);
+	void OnSeamPaintingChanged(ECheckBoxState InCheckState);
 	
 	/** Called to get the checked state of the enable flow checkbox */
-	ESlateCheckBoxState::Type IsEnableFlowChecked() const;
+	ECheckBoxState IsEnableFlowChecked() const;
 
 	/** Called when the enable flow checkbox value is changed */
-	void OnEnableFlowChanged(ESlateCheckBoxState::Type InCheckState);
+	void OnEnableFlowChanged(ECheckBoxState InCheckState);
 
 	/** Called when the erase weight is modified */
 	void OnEraseWeightChanged(int32 InWeightIndex);
@@ -165,10 +165,10 @@ protected:
 	void OnPaintWeightChanged(int32 InWeightIndex);
 
 	/** Called to get the checked state of the write channel check boxes */
-	ESlateCheckBoxState::Type IsWriteColorChannelChecked(EMeshPaintWriteColorChannels::Type CheckBoxInfo) const;
+	ECheckBoxState IsWriteColorChannelChecked(EMeshPaintWriteColorChannels::Type CheckBoxInfo) const;
 
 	/** Called when the state of one of the write channel check boxes is modified */
-	void OnWriteColorChannelChanged(ESlateCheckBoxState::Type InNewValue, EMeshPaintWriteColorChannels::Type CheckBoxInfo);
+	void OnWriteColorChannelChanged(ECheckBoxState InNewValue, EMeshPaintWriteColorChannels::Type CheckBoxInfo);
 
 	/** Called when the fill button is clicked */
 	FReply FillInstanceVertexColorsButtonClicked();
diff --git a/Engine/Source/Editor/NiagaraEditor/Private/SNiagaraEffectEditorWidget.h b/Engine/Source/Editor/NiagaraEditor/Private/SNiagaraEffectEditorWidget.h
--- a/Engine/Source/Editor/NiagaraEditor/Private/SNiagaraEffectEditorWidget.h
+++ b/Engine/Source/Editor/NiagaraEditor/Private/SNiagaraEffectEditorWidget.h
@@ -276,15 +276,15 @@ public:
 		return Emitter->GetProperties()->SpawnScript;
 	}
 
-	void OnEmitterEnabledChanged(ESlateCheckBoxState::Type NewCheckedState)
+	void OnEmitterEnabledChanged(ECheckBoxState NewCheckedState)
 	{
-		const bool bNewEnabledState = (NewCheckedState == ESlateCheckBoxState::Checked);
+		const bool bNewEnabledState = (NewCheckedState == ECheckBoxState::Checked);
 		Emitter->GetProperties()->bIsEnabled = bNewEnabledState;
 	}
 
-	ESlateCheckBoxState::Type IsEmitterEnabled() const
+	ECheckBoxState IsEmitterEnabled() const
 	{
-		return Emitter->IsEnabled() ? ESlateCheckBoxState::Checked : ESlateCheckBoxState::Unchecked;
+		return Emitter->IsEnabled() ? ECheckBoxState::Checked : ECheckBoxState::Unchecked;
 	}
 
 	void OnRenderModuleChanged(TSharedPtr<FString> ModName, ESelectInfo::Type SelectInfo);
diff --git a/Engine/Source/Editor/PListEditor/Private/PListNodeBoolean.cpp b/Engine/Source/Editor/PListEditor/Private/PListNodeBoolean.cpp
--- a/Engine/Source/Editor/PListEditor/Private/PListNodeBoolean.cpp
+++ b/Engine/Source/Editor/PListEditor/Private/PListNodeBoolean.cpp
@@ -113,7 +113,7 @@ TSharedRef<SWidget> FPListNodeBoolean::GenerateWidgetForColumn(const FName& Colu
 			.FillWidth(1.0f)
 			[
 				SAssignNew(ValueCheckBox, SCheckBox)
-				.IsChecked(bValue ? ESlateCheckBoxState::Checked : ESlateCheckBoxState::Unchecked)
+				.IsChecked(bValue ? ECheckBoxState::Checked : ECheckBoxState::Unchecked)
 				.OnCheckStateChanged(this, &FPListNodeBoolean::OnValueChanged)
 			]
 		];
@@ -292,9 +292,9 @@ void FPListNodeBoolean::OnKeyStringChanged(const FText& NewString)
 }
 
 /** Delegate: When the checkbox changes */
-void FPListNodeBoolean::OnValueChanged(ESlateCheckBoxState::Type NewValue)
+void FPListNodeBoolean::OnValueChanged(ECheckBoxState NewValue)
 {
-	if(NewValue == ESlateCheckBoxState::Checked)
+	if(NewValue == ECheckBoxState::Checked)
 	{
 		// Simply do nothing
 		if(bValue == true)
@@ -308,7 +308,7 @@ void FPListNodeBoolean::OnValueChanged(ESlateCheckBoxState::Type NewValue)
 		check(EditorWidget);
 		EditorWidget->MarkDirty();
 	}
-	else if(NewValue == ESlateCheckBoxState::Unchecked)
+	else if(NewValue == ECheckBoxState::Unchecked)
 	{
 		// Simply do nothing
 		if(bValue == false)
diff --git a/Engine/Source/Editor/PListEditor/Private/PListNodeBoolean.h b/Engine/Source/Editor/PListEditor/Private/PListNodeBoolean.h
--- a/Engine/Source/Editor/PListEditor/Private/PListNodeBoolean.h
+++ b/Engine/Source/Editor/PListEditor/Private/PListNodeBoolean.h
@@ -66,7 +66,7 @@ private:
 	/** Delegate: When key string changes */
 	void OnKeyStringChanged(const FText& NewString);
 	/** Delegate: When the checkbox changes */
-	void OnValueChanged(ESlateCheckBoxState::Type NewValue);
+	void OnValueChanged(ECheckBoxState NewValue);
 	/** Delegate: Changes the color of the key string text box based on validity */
 	FSlateColor GetKeyBackgroundColor() const;
 	/** Delegate: Changes the color of the key string text box based on validity */
diff --git a/Engine/Source/Editor/PackagesDialog/Private/PackagesDialog.cpp b/Engine/Source/Editor/PackagesDialog/Private/PackagesDialog.cpp
--- a/Engine/Source/Editor/PackagesDialog/Private/PackagesDialog.cpp
+++ b/Engine/Source/Editor/PackagesDialog/Private/PackagesDialog.cpp
@@ -127,23 +127,23 @@ void FPackagesDialogModule::SetMessage(const FText& InMessage)
  * @param	OutPackages	The array that should be populated with the desired packages
  * @param	InChecked	The type of packages that we want to retrieve
  */
-void FPackagesDialogModule::GetResults(OUT TArray<UPackage*>& OutPackages, ESlateCheckBoxState::Type InChecked)
+void FPackagesDialogModule::GetResults(OUT TArray<UPackage*>& OutPackages, ECheckBoxState InChecked)
 {
-	if(InChecked == ESlateCheckBoxState::Checked)
+	if(InChecked == ECheckBoxState::Checked)
 	{
 		for(int32 PackageIndex = 0; PackageIndex < CheckedPackages.Num(); ++PackageIndex)
 		{
 			OutPackages.Add(CheckedPackages[PackageIndex]);
 		}
 	}
-	else if(InChecked == ESlateCheckBoxState::Unchecked)
+	else if(InChecked == ECheckBoxState::Unchecked)
 	{
 		for(int32 PackageIndex = 0; PackageIndex < UncheckedPackages.Num(); ++PackageIndex)
 		{
 			OutPackages.Add(UncheckedPackages[PackageIndex]);
 		}
 	}
-	else if(InChecked == ESlateCheckBoxState::Undetermined)
+	else if(InChecked == ECheckBoxState::Undetermined)
 	{
 		for(int32 PackageIndex = 0; PackageIndex < UndeterminedPackages.Num(); ++PackageIndex)
 		{
@@ -162,7 +162,7 @@ void FPackagesDialogModule::GetResults(OUT TArray<UPackage*>& OutPackages, ESlat
  * @param	InIconName		The name of the icon to display
  * @param	InIconToolTip	The tooltip to display
  */
-void FPackagesDialogModule::AddPackageItem(UPackage* InPackage, const FString& InEntryName, ESlateCheckBoxState::Type InChecked, 
+void FPackagesDialogModule::AddPackageItem(UPackage* InPackage, const FString& InEntryName, ECheckBoxState InChecked, 
 	bool InDisabled/*=false*/, FString InIconName/*=TEXT("SavePackages.SCC_DlgNoIcon")*/, FString InIconToolTip/*=TEXT("")*/)
 {
 	PackagesDialogWidget.Get()->Add(MakeShareable(new FPackageItem(InPackage, InEntryName, InChecked, InDisabled, InIconName, InIconToolTip)));
diff --git a/Engine/Source/Editor/PackagesDialog/Private/SPackagesDialog.cpp b/Engine/Source/Editor/PackagesDialog/Private/SPackagesDialog.cpp
--- a/Engine/Source/Editor/PackagesDialog/Private/SPackagesDialog.cpp
+++ b/Engine/Source/Editor/PackagesDialog/Private/SPackagesDialog.cpp
@@ -242,11 +242,11 @@ EDialogReturnType SPackagesDialog::GetReturnType(OUT TArray<UPackage*>& OutCheck
 		for( int32 ItemIndex = 0; ItemIndex < Items.Num(); ++ItemIndex)
 		{
 			FPackageItem& item = *Items[ItemIndex];
-			if(item.GetState() == ESlateCheckBoxState::Checked)
+			if(item.GetState() == ECheckBoxState::Checked)
 			{
 				OutCheckedPackages.Add(item.GetPackage());
 			}
-			else if(item.GetState() == ESlateCheckBoxState::Unchecked)
+			else if(item.GetState() == ECheckBoxState::Unchecked)
 			{
 				OutUncheckedPackages.Add(item.GetPackage());
 			}
@@ -283,11 +283,11 @@ void SPackagesDialog::RefreshButtons()
 	for( int32 ItemIndex = 0; ItemIndex < Items.Num(); ++ItemIndex)
 	{
 		FPackageItem& item = *Items[ItemIndex];
-		if(item.GetState() == ESlateCheckBoxState::Checked)
+		if(item.GetState() == ECheckBoxState::Checked)
 		{
 			CheckedItems++;
 		}
-		else if(item.GetState() == ESlateCheckBoxState::Unchecked)
+		else if(item.GetState() == ECheckBoxState::Unchecked)
 		{
 			UncheckedItems++;
 		}
@@ -463,48 +463,48 @@ TArray< TSharedPtr<FPackageItem> > SPackagesDialog::GetSelectedItems( bool bAllI
 	return SelectedItems;
 }
 
-ESlateCheckBoxState::Type SPackagesDialog::GetToggleSelectedState() const
+ECheckBoxState SPackagesDialog::GetToggleSelectedState() const
 {
 	//default to a Checked state
-	ESlateCheckBoxState::Type PendingState = ESlateCheckBoxState::Checked;
+	ECheckBoxState PendingState = ECheckBoxState::Checked;
 
 	TArray< TSharedPtr<FPackageItem> > SelectedItems = GetSelectedItems( true );
 
 	//Iterate through the list of selected packages
 	for ( auto SelectedItem = SelectedItems.CreateConstIterator(); SelectedItem; ++SelectedItem )
 	{
-		if ( SelectedItem->Get()->GetState() == ESlateCheckBoxState::Unchecked )
+		if ( SelectedItem->Get()->GetState() == ECheckBoxState::Unchecked )
 		{
 			//if any package in the selection is Unchecked, then represent the entire set of highlighted packages as Unchecked,
 			//so that the first (user) toggle of ToggleSelectedCheckBox consistently Checks all highlighted packages
-			PendingState = ESlateCheckBoxState::Unchecked;
+			PendingState = ECheckBoxState::Unchecked;
 		}
 	}
 
 	return PendingState;
 }
 
-void SPackagesDialog::OnToggleSelectedCheckBox(ESlateCheckBoxState::Type InNewState)
+void SPackagesDialog::OnToggleSelectedCheckBox(ECheckBoxState InNewState)
 {
 	TArray< TSharedPtr<FPackageItem> > SelectedItems = GetSelectedItems( true );
 
 	for ( auto SelectedItem = SelectedItems.CreateConstIterator(); SelectedItem; ++SelectedItem )
 	{
 		FPackageItem *item = (*SelectedItem).Get();
-		if (InNewState == ESlateCheckBoxState::Checked)
+		if (InNewState == ECheckBoxState::Checked)
 		{
 			if(item->IsDisabled())
 			{
-				item->SetState(ESlateCheckBoxState::Undetermined);
+				item->SetState(ECheckBoxState::Undetermined);
 			}
 			else
 			{
-				item->SetState(ESlateCheckBoxState::Checked);
+				item->SetState(ECheckBoxState::Checked);
 			}
 		}
 		else
 		{
-			item->SetState(ESlateCheckBoxState::Unchecked);
+			item->SetState(ECheckBoxState::Unchecked);
 		}
 	}
 
@@ -553,7 +553,7 @@ void SPackagesDialog::PopulateIgnoreForSaveItems( const TSet<FString>& InIgnoreP
 	{
 		const FString& ItemName = (*ItItem)->GetName();
 
-		const ESlateCheckBoxState::Type CheckedStatus = (InIgnorePackages.Find(ItemName) != NULL) ? ESlateCheckBoxState::Unchecked : ESlateCheckBoxState::Checked;
+		const ECheckBoxState CheckedStatus = (InIgnorePackages.Find(ItemName) != NULL) ? ECheckBoxState::Unchecked : ECheckBoxState::Checked;
 
 		if (!(*ItItem)->IsDisabled())
 		{
@@ -566,7 +566,7 @@ void SPackagesDialog::PopulateIgnoreForSaveArray( OUT TSet<FString>& InOutIgnore
 {
 	for( auto ItItem=Items.CreateConstIterator(); ItItem; ++ItItem )
 	{
-		if((*ItItem)->GetState()==ESlateCheckBoxState::Unchecked)
+		if((*ItItem)->GetState()==ECheckBoxState::Unchecked)
 		{
 			InOutIgnorePackages.Add( (*ItItem)->GetName() );
 		}
diff --git a/Engine/Source/Editor/PackagesDialog/Private/SPackagesDialog.h b/Engine/Source/Editor/PackagesDialog/Private/SPackagesDialog.h
--- a/Engine/Source/Editor/PackagesDialog/Private/SPackagesDialog.h
+++ b/Engine/Source/Editor/PackagesDialog/Private/SPackagesDialog.h
@@ -107,7 +107,7 @@ private:
 class FPackageItem : public TSharedFromThis<FPackageItem>
 {
 public:
-	FPackageItem(UPackage* InPackage, const FString& InEntryName, ESlateCheckBoxState::Type InState, bool InDisabled = false, FString InIconName=TEXT("SavePackages.SCC_DlgNoIcon"), FString InIconToolTip=TEXT(""))
+	FPackageItem(UPackage* InPackage, const FString& InEntryName, ECheckBoxState InState, bool InDisabled = false, FString InIconName=TEXT("SavePackages.SCC_DlgNoIcon"), FString InIconToolTip=TEXT(""))
 		: Package(InPackage)
 		, EntryName(InEntryName)
 		, State(InState)
@@ -116,9 +116,9 @@ public:
 		, IconToolTip(InIconToolTip)
 	{
 		/** if the item is checked and disabled make the state Undetermined */
-		if(State == ESlateCheckBoxState::Checked && Disabled)
+		if(State == ECheckBoxState::Checked && Disabled)
 		{
-			State = ESlateCheckBoxState::Undetermined;
+			State = ECheckBoxState::Undetermined;
 		}
 	}
 
@@ -127,7 +127,7 @@ public:
 	 *
 	 * @return the item state
 	 */
-	ESlateCheckBoxState::Type OnGetDisplayCheckState() const
+	ECheckBoxState OnGetDisplayCheckState() const
 	{
 		RefreshButtonCallback.ExecuteIfBound();
 		return State;
@@ -138,14 +138,14 @@ public:
 	 *
 	 * @param InNewState the new state to set the item to
 	 */
-	void OnDisplayCheckStateChanged(ESlateCheckBoxState::Type InNewState)
+	void OnDisplayCheckStateChanged(ECheckBoxState InNewState)
 	{
 		State = InNewState;
 
 		/** if the item is checked and disabled make the state Undetermined */
-		if( State == ESlateCheckBoxState::Checked  && Disabled )
+		if( State == ECheckBoxState::Checked  && Disabled )
 		{
-			State = ESlateCheckBoxState::Undetermined;
+			State = ECheckBoxState::Undetermined;
 		}
 	}
 
@@ -164,7 +164,7 @@ public:
 	 *
 	 * @return the state of the checkbox item
 	 */
-	ESlateCheckBoxState::Type GetState() const { return State; }
+	ECheckBoxState GetState() const { return State; }
 
 	/**
 	 * Gets the package represented by this checkbox item
@@ -236,12 +236,12 @@ public:
 	 *
 	 * @param	NewState	New state
 	 */
-	void SetState(ESlateCheckBoxState::Type NewState) { State = NewState; }
+	void SetState(ECheckBoxState NewState) { State = NewState; }
 
 private:
 	UPackage* Package; 						// The package associated with this entry
 	FString EntryName;						// Name of the checkbox
-	ESlateCheckBoxState::Type State;		// The state of the checkbox
+	ECheckBoxState State;		// The state of the checkbox
 	bool Disabled; 							// if the entry is disabled
 	FString IconName; 						// Name of an icon to show next to the checkbox
 	FString IconToolTip;					// ToolTip to display for the icon
@@ -369,13 +369,13 @@ private:
 	 * @return	the desired toggle state for the ToggleSelectedCheckBox.
 	 * Returns Unchecked, unless all of the selected packages are Checked.
 	 */
-	ESlateCheckBoxState::Type GetToggleSelectedState() const;
+	ECheckBoxState GetToggleSelectedState() const;
 
 	/** 
 	 * Toggles the highlighted packages.
 	 * If no packages are explicitly highlighted, toggles all packages in the list.
 	 */
-	void OnToggleSelectedCheckBox(ESlateCheckBoxState::Type InNewState);
+	void OnToggleSelectedCheckBox(ECheckBoxState InNewState);
 
 	/** 
 	 * Makes the widget for the checkbox items in the list view 
diff --git a/Engine/Source/Editor/PackagesDialog/Public/PackagesDialog.h b/Engine/Source/Editor/PackagesDialog/Public/PackagesDialog.h
--- a/Engine/Source/Editor/PackagesDialog/Public/PackagesDialog.h
+++ b/Engine/Source/Editor/PackagesDialog/Public/PackagesDialog.h
@@ -78,7 +78,7 @@ public:
 	 * @param	OutPackages		The array that should be populated with the desired packages
 	 * @param	InChecked		The type of packages that we want to retrieve
 	 */
-	virtual void GetResults(TArray<UPackage*>& OutPackages, ESlateCheckBoxState::Type InChecked);
+	virtual void GetResults(TArray<UPackage*>& OutPackages, ECheckBoxState InChecked);
 
 	/**
 	 * Adds a new item to the checkbox that represents a package
@@ -90,7 +90,7 @@ public:
 	 * @param	InIconName		The name of the icon to display
 	 * @param	InIconToolTip	The tooltip to display
 	 */
-	virtual void AddPackageItem(UPackage* InPackage, const FString& InEntryName, ESlateCheckBoxState::Type InChecked, 
+	virtual void AddPackageItem(UPackage* InPackage, const FString& InEntryName, ECheckBoxState InChecked, 
 								bool InDisabled = false, FString InIconName=TEXT("SavePackages.SCC_DlgNoIcon"), FString InIconToolTip=TEXT(""));
 
 	/**
diff --git a/Engine/Source/Editor/Persona/Private/ApexClothingOptionWindow.h b/Engine/Source/Editor/Persona/Private/ApexClothingOptionWindow.h
--- a/Engine/Source/Editor/Persona/Private/ApexClothingOptionWindow.h
+++ b/Engine/Source/Editor/Persona/Private/ApexClothingOptionWindow.h
@@ -41,14 +41,14 @@ public:
 		return FReply::Handled();
 	}
 
-	ESlateCheckBoxState::Type IsCheckedLOD() const
+	ECheckBoxState IsCheckedLOD() const
 	{
-		return bUseLOD ? ESlateCheckBoxState::Checked : ESlateCheckBoxState::Unchecked;
+		return bUseLOD ? ECheckBoxState::Checked : ECheckBoxState::Unchecked;
 	}
 
-	void OnUseLOD(ESlateCheckBoxState::Type CheckState)
+	void OnUseLOD(ECheckBoxState CheckState)
 	{
-		bUseLOD = (CheckState == ESlateCheckBoxState::Checked);
+		bUseLOD = (CheckState == ECheckBoxState::Checked);
 	}
 
 	bool CanImport() const
diff --git a/Engine/Source/Editor/Persona/Private/PersonaMeshDetails.cpp b/Engine/Source/Editor/Persona/Private/PersonaMeshDetails.cpp
--- a/Engine/Source/Editor/Persona/Private/PersonaMeshDetails.cpp
+++ b/Engine/Source/Editor/Persona/Private/PersonaMeshDetails.cpp
@@ -356,9 +356,9 @@ float FSkelMeshReductionSettingsLayout::GetWeldingThreshold() const
 	return ReductionSettings.WeldingThreshold;
 }
 
-ESlateCheckBoxState::Type FSkelMeshReductionSettingsLayout::ShouldRecalculateNormals() const
+ECheckBoxState FSkelMeshReductionSettingsLayout::ShouldRecalculateNormals() const
 {
-	return ReductionSettings.bRecalcNormals ? ESlateCheckBoxState::Checked : ESlateCheckBoxState::Unchecked;
+	return ReductionSettings.bRecalcNormals ? ECheckBoxState::Checked : ECheckBoxState::Unchecked;
 }
 
 float FSkelMeshReductionSettingsLayout::GetHardAngleThreshold() const
@@ -387,9 +387,9 @@ void FSkelMeshReductionSettingsLayout::OnWeldingThresholdChanged(float NewValue)
 	ReductionSettings.WeldingThreshold = NewValue;
 }
 
-void FSkelMeshReductionSettingsLayout::OnRecalculateNormalsChanged(ESlateCheckBoxState::Type NewValue)
+void FSkelMeshReductionSettingsLayout::OnRecalculateNormalsChanged(ECheckBoxState NewValue)
 {
-	ReductionSettings.bRecalcNormals = NewValue == ESlateCheckBoxState::Checked;
+	ReductionSettings.bRecalcNormals = NewValue == ECheckBoxState::Checked;
 }
 
 void FSkelMeshReductionSettingsLayout::OnHardAngleThresholdChanged(float NewValue)
@@ -1046,29 +1046,29 @@ TSharedRef<SWidget> FPersonaMeshDetails::OnGenerateCustomMaterialWidgetsForMater
 	return MaterialWidget;
 }
 
-ESlateCheckBoxState::Type FPersonaMeshDetails::IsSectionSelected(int32 SectionIndex) const
+ECheckBoxState FPersonaMeshDetails::IsSectionSelected(int32 SectionIndex) const
 {
-	ESlateCheckBoxState::Type State = ESlateCheckBoxState::Unchecked;
+	ECheckBoxState State = ECheckBoxState::Unchecked;
 	USkeletalMesh* Mesh = Cast<USkeletalMesh>(SelectedObjects[0].Get());
 
 	if (Mesh)
 	{
-		State = Mesh->SelectedEditorSection == SectionIndex ? ESlateCheckBoxState::Checked : ESlateCheckBoxState::Unchecked;
+		State = Mesh->SelectedEditorSection == SectionIndex ? ECheckBoxState::Checked : ECheckBoxState::Unchecked;
 	}
 
 	return State;
 }
 
-void FPersonaMeshDetails::OnSectionSelectedChanged(ESlateCheckBoxState::Type NewState, int32 SectionIndex)
+void FPersonaMeshDetails::OnSectionSelectedChanged(ECheckBoxState NewState, int32 SectionIndex)
 {
 	USkeletalMesh* Mesh = Cast<USkeletalMesh>(SelectedObjects[0].Get());
 	if (Mesh)
 	{
-		if (NewState == ESlateCheckBoxState::Checked)
+		if (NewState == ECheckBoxState::Checked)
 		{
 			Mesh->SelectedEditorSection = SectionIndex;
 		}
-		else if (NewState == ESlateCheckBoxState::Unchecked)
+		else if (NewState == ECheckBoxState::Unchecked)
 		{
 			Mesh->SelectedEditorSection = INDEX_NONE;
 		}
@@ -1076,32 +1076,32 @@ void FPersonaMeshDetails::OnSectionSelectedChanged(ESlateCheckBoxState::Type New
 	}
 }
 
-ESlateCheckBoxState::Type FPersonaMeshDetails::IsShadowCastingEnabled(int32 MaterialIndex) const
+ECheckBoxState FPersonaMeshDetails::IsShadowCastingEnabled(int32 MaterialIndex) const
 {
-	ESlateCheckBoxState::Type State = ESlateCheckBoxState::Unchecked;
+	ECheckBoxState State = ECheckBoxState::Unchecked;
 	USkeletalMesh* Mesh = Cast<USkeletalMesh>( SelectedObjects[0].Get() );
 
 	if (Mesh && MaterialIndex < Mesh->Materials.Num())
 	{
-		State = Mesh->Materials[MaterialIndex].bEnableShadowCasting ? ESlateCheckBoxState::Checked : ESlateCheckBoxState::Unchecked;
+		State = Mesh->Materials[MaterialIndex].bEnableShadowCasting ? ECheckBoxState::Checked : ECheckBoxState::Unchecked;
 	}
 
 	return State;
 }
 
-void FPersonaMeshDetails::OnShadowCastingChanged(ESlateCheckBoxState::Type NewState, int32 MaterialIndex)
+void FPersonaMeshDetails::OnShadowCastingChanged(ECheckBoxState NewState, int32 MaterialIndex)
 {
 	USkeletalMesh* Mesh = Cast<USkeletalMesh>( SelectedObjects[0].Get() );
 
 	if ( Mesh )
 	{
-		if (NewState == ESlateCheckBoxState::Checked)
+		if (NewState == ECheckBoxState::Checked)
 		{
 			const FScopedTransaction Transaction( LOCTEXT( "SetShadowCastingFlag", "Set Shadow Casting For Material" ) );
 			Mesh->Modify();
 			Mesh->Materials[MaterialIndex].bEnableShadowCasting = true;
 		}
-		else if (NewState == ESlateCheckBoxState::Unchecked)
+		else if (NewState == ECheckBoxState::Unchecked)
 		{
 			const FScopedTransaction Transaction( LOCTEXT( "ClearShadowCastingFlag", "Clear Shadow Casting For Material" ) );
 			Mesh->Modify();
diff --git a/Engine/Source/Editor/Persona/Private/PersonaMeshDetails.h b/Engine/Source/Editor/Persona/Private/PersonaMeshDetails.h
--- a/Engine/Source/Editor/Persona/Private/PersonaMeshDetails.h
+++ b/Engine/Source/Editor/Persona/Private/PersonaMeshDetails.h
@@ -60,14 +60,14 @@ private:
 	float GetPercentTriangles() const;
 	float GetMaxDeviation() const;
 	float GetWeldingThreshold() const;
-	ESlateCheckBoxState::Type ShouldRecalculateNormals() const;
+	ECheckBoxState ShouldRecalculateNormals() const;
 	float GetHardAngleThreshold() const;
 	int32 GetMaxBonesPerVertex() const;
 
 	void OnPercentTrianglesChanged(float NewValue);
 	void OnMaxDeviationChanged(float NewValue);
 	void OnReductionAmountChanged(float NewValue);
-	void OnRecalculateNormalsChanged(ESlateCheckBoxState::Type NewValue);
+	void OnRecalculateNormalsChanged(ECheckBoxState NewValue);
 	void OnWeldingThresholdChanged(float NewValue);
 	void OnHardAngleThresholdChanged(float NewValue);
 	void OnMaxBonesPerVertexChanged(int32 NewValue);
@@ -145,28 +145,28 @@ private:
 	 *
 	 * @param SectionIndex	The material section that is being tested
 	 */
-	ESlateCheckBoxState::Type IsSectionSelected(int32 SectionIndex) const;
+	ECheckBoxState IsSectionSelected(int32 SectionIndex) const;
 
 	/**
 	 * Handler for changing highlight status on a material
 	 *
 	 * @param SectionIndex	The material section that is being tested
 	 */
-	void OnSectionSelectedChanged(ESlateCheckBoxState::Type NewState, int32 SectionIndex);
+	void OnSectionSelectedChanged(ECheckBoxState NewState, int32 SectionIndex);
 
 		/**
 	 * Handler for check box display based on whether the material has shadow casting enabled
 	 *
 	 * @param MaterialIndex	The material index which a section in a specific LOD model has
 	 */
-	ESlateCheckBoxState::Type IsShadowCastingEnabled(int32 MaterialIndex) const;
+	ECheckBoxState IsShadowCastingEnabled(int32 MaterialIndex) const;
 
 	/**
 	 * Handler for changing shadow casting status on a material
 	 *
 	 * @param MaterialIndex	The material index which a section in a specific LOD model has
 	 */
-	void OnShadowCastingChanged(ESlateCheckBoxState::Type NewState, int32 MaterialIndex);
+	void OnShadowCastingChanged(ECheckBoxState NewState, int32 MaterialIndex);
 
 	/**
 	 * Handler for enabling delete button on materials
diff --git a/Engine/Source/Editor/Persona/Private/SAnimCurvePanel.cpp b/Engine/Source/Editor/Persona/Private/SAnimCurvePanel.cpp
--- a/Engine/Source/Editor/Persona/Private/SAnimCurvePanel.cpp
+++ b/Engine/Source/Editor/Persona/Private/SAnimCurvePanel.cpp
@@ -159,10 +159,10 @@ public:
 	void DeleteTrack();
 
 	// Sets the current mode for this curve
-	void ToggleCurveMode(ESlateCheckBoxState::Type NewState, EAnimCurveFlags ModeToSet);
+	void ToggleCurveMode(ECheckBoxState NewState, EAnimCurveFlags ModeToSet);
 
 	// Returns whether this curve is of the specificed mode type
-	ESlateCheckBoxState::Type IsCurveOfMode(EAnimCurveFlags ModeToTest) const;
+	ECheckBoxState IsCurveOfMode(EAnimCurveFlags ModeToTest) const;
 
 	/**
 	 * Build and display curve track context menu.
@@ -171,8 +171,8 @@ public:
 	FReply OnContextMenu();
 
 	// expand editor mode 
-	ESlateCheckBoxState::Type IsEditorExpanded() const;
-	void ToggleExpandEditor(ESlateCheckBoxState::Type NewType);
+	ECheckBoxState IsEditorExpanded() const;
+	void ToggleExpandEditor(ECheckBoxState NewType);
 	const FSlateBrush* GetExpandContent() const;
 	FVector2D GetDesiredSize() const;
 
@@ -357,7 +357,7 @@ void SCurveEdTrack::DeleteTrack()
 	}
 }
 
-void SCurveEdTrack::ToggleCurveMode(ESlateCheckBoxState::Type NewState,EAnimCurveFlags ModeToSet)
+void SCurveEdTrack::ToggleCurveMode(ECheckBoxState NewState,EAnimCurveFlags ModeToSet)
 {
 	const int32 AllModes = (ACF_DrivesMorphTarget|ACF_DrivesMaterial);
 	check((ModeToSet&AllModes) != 0); //unexpected value for ModeToSet
@@ -366,7 +366,7 @@ void SCurveEdTrack::ToggleCurveMode(ESlateCheckBoxState::Type NewState,EAnimCurv
 
 	FText UndoLabel;
 	bool bIsSwitchingFlagOn = !CurveData->GetCurveTypeFlag(ModeToSet);
-	check(bIsSwitchingFlagOn == (NewState==ESlateCheckBoxState::Checked));
+	check(bIsSwitchingFlagOn == (NewState==ECheckBoxState::Checked));
 	
 	if(bIsSwitchingFlagOn)
 	{
@@ -398,10 +398,10 @@ void SCurveEdTrack::ToggleCurveMode(ESlateCheckBoxState::Type NewState,EAnimCurv
 	CurveData->SetCurveTypeFlag(ModeToSet, bIsSwitchingFlagOn);
 }
 
-ESlateCheckBoxState::Type SCurveEdTrack::IsCurveOfMode(EAnimCurveFlags ModeToTest) const
+ECheckBoxState SCurveEdTrack::IsCurveOfMode(EAnimCurveFlags ModeToTest) const
 {
 	FFloatCurve* CurveData = (FFloatCurve*)(CurveInterface->CurveData);
-	return CurveData->GetCurveTypeFlag(ModeToTest) ? ESlateCheckBoxState::Checked : ESlateCheckBoxState::Unchecked;
+	return CurveData->GetCurveTypeFlag(ModeToTest) ? ECheckBoxState::Checked : ECheckBoxState::Unchecked;
 }
 
 FReply SCurveEdTrack::OnContextMenu()
@@ -419,14 +419,14 @@ FReply SCurveEdTrack::OnContextMenu()
 	return FReply::Handled();
 }
 
-ESlateCheckBoxState::Type SCurveEdTrack::IsEditorExpanded() const
+ECheckBoxState SCurveEdTrack::IsEditorExpanded() const
 {
-	return (bUseExpandEditor)? ESlateCheckBoxState::Checked : ESlateCheckBoxState::Unchecked;
+	return (bUseExpandEditor)? ECheckBoxState::Checked : ECheckBoxState::Unchecked;
 }
 
-void SCurveEdTrack::ToggleExpandEditor(ESlateCheckBoxState::Type NewType)
+void SCurveEdTrack::ToggleExpandEditor(ECheckBoxState NewType)
 {
-	bUseExpandEditor = (NewType == ESlateCheckBoxState::Checked);
+	bUseExpandEditor = (NewType == ECheckBoxState::Checked);
 }
 
 FVector2D SCurveEdTrack::GetDesiredSize() const
@@ -696,16 +696,16 @@ EVisibility SAnimCurvePanel::IsSetAllTracksButtonVisible() const
 	return (Tracks.Num() > 1) ? EVisibility::Visible : EVisibility::Hidden;
 }
 
-void SAnimCurvePanel::ToggleAllCurveModes(ESlateCheckBoxState::Type NewState, EAnimCurveFlags ModeToSet)
+void SAnimCurvePanel::ToggleAllCurveModes(ECheckBoxState NewState, EAnimCurveFlags ModeToSet)
 {
-	const ESlateCheckBoxState::Type CurrentAllState = AreAllCurvesOfMode(ModeToSet);
+	const ECheckBoxState CurrentAllState = AreAllCurvesOfMode(ModeToSet);
 	for(TWeakPtr<SCurveEdTrack> TrackWeak : Tracks)
 	{
 		TSharedPtr<SCurveEdTrack> TrackWidget = TrackWeak.Pin();
 		if( TrackWidget.IsValid() )
 		{
-			const ESlateCheckBoxState::Type CurrentTrackState = TrackWidget->IsCurveOfMode(ModeToSet);
-			if( (CurrentAllState == CurrentTrackState) || ((CurrentAllState == ESlateCheckBoxState::Undetermined) && (CurrentTrackState == ESlateCheckBoxState::Unchecked)) )
+			const ECheckBoxState CurrentTrackState = TrackWidget->IsCurveOfMode(ModeToSet);
+			if( (CurrentAllState == CurrentTrackState) || ((CurrentAllState == ECheckBoxState::Undetermined) && (CurrentTrackState == ECheckBoxState::Unchecked)) )
 			{
 				TrackWidget->ToggleCurveMode( NewState, ModeToSet );
 			}
@@ -713,7 +713,7 @@ void SAnimCurvePanel::ToggleAllCurveModes(ESlateCheckBoxState::Type NewState, EA
 	}
 }
 
-ESlateCheckBoxState::Type SAnimCurvePanel::AreAllCurvesOfMode(EAnimCurveFlags ModeToSet) const
+ECheckBoxState SAnimCurvePanel::AreAllCurvesOfMode(EAnimCurveFlags ModeToSet) const
 {
 	int32 NumChecked = 0;
 	for(const TWeakPtr<SCurveEdTrack> TrackWeak : Tracks)
@@ -721,7 +721,7 @@ ESlateCheckBoxState::Type SAnimCurvePanel::AreAllCurvesOfMode(EAnimCurveFlags Mo
 		const TSharedPtr<SCurveEdTrack> TrackWidget = TrackWeak.Pin();
 		if( TrackWidget.IsValid() )
 		{
-			if( TrackWidget->IsCurveOfMode(ModeToSet) )
+			if ( TrackWidget->IsCurveOfMode(ModeToSet) == ECheckBoxState::Checked )
 			{
 				NumChecked++;
 			}
@@ -729,13 +729,13 @@ ESlateCheckBoxState::Type SAnimCurvePanel::AreAllCurvesOfMode(EAnimCurveFlags Mo
 	}
 	if( NumChecked == Tracks.Num() )
 	{
-		return ESlateCheckBoxState::Checked;
+		return ECheckBoxState::Checked;
 	}
 	else if( NumChecked == 0 )
 	{
-		return ESlateCheckBoxState::Unchecked;
+		return ECheckBoxState::Unchecked;
 	}
-	return ESlateCheckBoxState::Undetermined;
+	return ECheckBoxState::Undetermined;
 }
 
 void SAnimCurvePanel::UpdatePanel()
@@ -935,23 +935,23 @@ TSharedRef<SWidget> SAnimCurvePanel::GenerateCurveList()
 	return NewWidget;
 }
 
-ESlateCheckBoxState::Type SAnimCurvePanel::IsCurveEditable(USkeleton::AnimCurveUID Uid) const
+ECheckBoxState SAnimCurvePanel::IsCurveEditable(USkeleton::AnimCurveUID Uid) const
 {
 	if ( Sequence )
 	{
 		const FFloatCurve* Curve = static_cast<const FFloatCurve *>(Sequence->RawCurveData.GetCurveData(Uid, FRawCurveTracks::FloatType));
 		if ( Curve )
 		{
-			return Curve->GetCurveTypeFlag(ACF_Editable)? ESlateCheckBoxState::Checked : ESlateCheckBoxState::Unchecked;
+			return Curve->GetCurveTypeFlag(ACF_Editable)? ECheckBoxState::Checked : ECheckBoxState::Unchecked;
 		}
 	}
 
-	return ESlateCheckBoxState::Undetermined;
+	return ECheckBoxState::Undetermined;
 }
 
-void SAnimCurvePanel::ToggleEditability(ESlateCheckBoxState::Type NewType, USkeleton::AnimCurveUID Uid)
+void SAnimCurvePanel::ToggleEditability(ECheckBoxState NewType, USkeleton::AnimCurveUID Uid)
 {
-	bool bEdit = (NewType == ESlateCheckBoxState::Checked);
+	bool bEdit = (NewType == ECheckBoxState::Checked);
 
 	if ( Sequence )
 	{
@@ -1199,14 +1199,14 @@ TSharedRef<SWidget> SAnimCurvePanel::CreateCurveContextMenu(FFloatCurve* Curve)
 	return MenuBuilder.MakeWidget();
 }
 
-ESlateCheckBoxState::Type SAnimCurvePanel::GetCurveFlagAsCheckboxState(FFloatCurve* Curve, EAnimCurveFlags InFlag) const
+ECheckBoxState SAnimCurvePanel::GetCurveFlagAsCheckboxState(FFloatCurve* Curve, EAnimCurveFlags InFlag) const
 {
-	return Curve->GetCurveTypeFlag(InFlag) ? ESlateCheckBoxState::Checked : ESlateCheckBoxState::Unchecked;
+	return Curve->GetCurveTypeFlag(InFlag) ? ECheckBoxState::Checked : ECheckBoxState::Unchecked;
 }
 
-void SAnimCurvePanel::SetCurveFlagFromCheckboxState(ESlateCheckBoxState::Type CheckState, FFloatCurve* Curve, EAnimCurveFlags InFlag)
+void SAnimCurvePanel::SetCurveFlagFromCheckboxState(ECheckBoxState CheckState, FFloatCurve* Curve, EAnimCurveFlags InFlag)
 {
-	bool Enabled = CheckState == ESlateCheckBoxState::Checked;
+	bool Enabled = CheckState == ECheckBoxState::Checked;
 	Curve->SetCurveTypeFlag(InFlag, Enabled);
 }
 
diff --git a/Engine/Source/Editor/Persona/Private/SAnimCurvePanel.h b/Engine/Source/Editor/Persona/Private/SAnimCurvePanel.h
--- a/Engine/Source/Editor/Persona/Private/SAnimCurvePanel.h
+++ b/Engine/Source/Editor/Persona/Private/SAnimCurvePanel.h
@@ -98,12 +98,12 @@ public:
 	/**
 	 * Toggle all the curve modes of all the tracks
 	 */
-	void ToggleAllCurveModes(ESlateCheckBoxState::Type NewState, EAnimCurveFlags ModeToSet);
+	void ToggleAllCurveModes(ECheckBoxState NewState, EAnimCurveFlags ModeToSet);
 
 	/**
 	 * Are all the curve modes of all the tracks set to the mode
 	 */
-	ESlateCheckBoxState::Type AreAllCurvesOfMode(EAnimCurveFlags ModeToTest) const;
+	ECheckBoxState AreAllCurvesOfMode(EAnimCurveFlags ModeToTest) const;
 
 	/**
 	 * Visibility of the set-all-tracks button
@@ -143,11 +143,11 @@ private:
 	/**
 	 * Returns true if this curve is editable
 	 */
-	ESlateCheckBoxState::Type	IsCurveEditable(USkeleton::AnimCurveUID Uid) const;
+	ECheckBoxState	IsCurveEditable(USkeleton::AnimCurveUID Uid) const;
 	/**
 	 * Toggle curve visibility
 	 */
-	void		ToggleEditability(ESlateCheckBoxState::Type NewType, USkeleton::AnimCurveUID Uid);
+	void		ToggleEditability(ECheckBoxState NewType, USkeleton::AnimCurveUID Uid);
 	/**
 	 * Refresh Panel
 	 */
@@ -192,12 +192,12 @@ private:
 	/**
 	 * Convert the requested flag bool value into a checkbox state
 	 */
-	ESlateCheckBoxState::Type GetCurveFlagAsCheckboxState(FFloatCurve* Curve, EAnimCurveFlags InFlag) const;
+	ECheckBoxState GetCurveFlagAsCheckboxState(FFloatCurve* Curve, EAnimCurveFlags InFlag) const;
 
 	/**
 	 * Convert a given checkbox state into a flag value in the provided curve
 	 */
-	void SetCurveFlagFromCheckboxState(ESlateCheckBoxState::Type CheckState, FFloatCurve* Curve, EAnimCurveFlags InFlag);
+	void SetCurveFlagFromCheckboxState(ECheckBoxState CheckState, FFloatCurve* Curve, EAnimCurveFlags InFlag);
 
 	/**
 	 * Handler for converting a curve from metadata to variable curve or vice versa
diff --git a/Engine/Source/Editor/Persona/Private/SAnimTrackCurvePanel.cpp b/Engine/Source/Editor/Persona/Private/SAnimTrackCurvePanel.cpp
--- a/Engine/Source/Editor/Persona/Private/SAnimTrackCurvePanel.cpp
+++ b/Engine/Source/Editor/Persona/Private/SAnimTrackCurvePanel.cpp
@@ -219,8 +219,8 @@ public:
 	FReply OnContextMenu();
 
 	// expand editor mode 
-	ESlateCheckBoxState::Type IsEditorExpanded() const;
-	void ToggleExpandEditor(ESlateCheckBoxState::Type NewType);
+	ECheckBoxState IsEditorExpanded() const;
+	void ToggleExpandEditor(ECheckBoxState NewType);
 	const FSlateBrush* GetExpandContent() const;
 	FVector2D GetDesiredSize() const;
 
@@ -738,23 +738,23 @@ TSharedRef<SWidget> SAnimTrackCurvePanel::GenerateCurveList()
 	return NewWidget;
 }
 
-ESlateCheckBoxState::Type SAnimTrackCurvePanel::IsCurveEditable(USkeleton::AnimCurveUID Uid) const
+ECheckBoxState SAnimTrackCurvePanel::IsCurveEditable(USkeleton::AnimCurveUID Uid) const
 {
 	if ( Sequence )
 	{
 		const FTransformCurve* Curve = static_cast<const FTransformCurve *>(Sequence->RawCurveData.GetCurveData(Uid, FRawCurveTracks::TransformType));
 		if ( Curve )
 		{
-			return Curve->GetCurveTypeFlag(ACF_Editable)? ESlateCheckBoxState::Checked : ESlateCheckBoxState::Unchecked;
+			return Curve->GetCurveTypeFlag(ACF_Editable)? ECheckBoxState::Checked : ECheckBoxState::Unchecked;
 		}
 	}
 
-	return ESlateCheckBoxState::Undetermined;
+	return ECheckBoxState::Undetermined;
 }
 
-void SAnimTrackCurvePanel::ToggleEditability(ESlateCheckBoxState::Type NewType, USkeleton::AnimCurveUID Uid)
+void SAnimTrackCurvePanel::ToggleEditability(ECheckBoxState NewType, USkeleton::AnimCurveUID Uid)
 {
-	bool bEdit = (NewType == ESlateCheckBoxState::Checked);
+	bool bEdit = (NewType == ECheckBoxState::Checked);
 
 	if ( Sequence )
 	{
@@ -788,15 +788,15 @@ FReply		SAnimTrackCurvePanel::ShowAll(bool bShow)
 	return FReply::Handled();
 }
 
-ESlateCheckBoxState::Type SAnimTrackCurvePanel::GetCurveFlagAsCheckboxState(USkeleton::AnimCurveUID CurveUid, EAnimCurveFlags InFlag) const
+ECheckBoxState SAnimTrackCurvePanel::GetCurveFlagAsCheckboxState(USkeleton::AnimCurveUID CurveUid, EAnimCurveFlags InFlag) const
 {
 	FAnimCurveBase* Curve = Sequence->RawCurveData.GetCurveData(CurveUid, FRawCurveTracks::TransformType);
-	return Curve && Curve->GetCurveTypeFlag(InFlag) ? ESlateCheckBoxState::Checked : ESlateCheckBoxState::Unchecked;
+	return Curve && Curve->GetCurveTypeFlag(InFlag) ? ECheckBoxState::Checked : ECheckBoxState::Unchecked;
 }
 
-void SAnimTrackCurvePanel::SetCurveFlagFromCheckboxState(ESlateCheckBoxState::Type CheckState, USkeleton::AnimCurveUID CurveUid, EAnimCurveFlags InFlag)
+void SAnimTrackCurvePanel::SetCurveFlagFromCheckboxState(ECheckBoxState CheckState, USkeleton::AnimCurveUID CurveUid, EAnimCurveFlags InFlag)
 {
-	bool Enabled = CheckState == ESlateCheckBoxState::Checked;
+	bool Enabled = CheckState == ECheckBoxState::Checked;
 	FAnimCurveBase* Curve = Sequence->RawCurveData.GetCurveData(CurveUid, FRawCurveTracks::TransformType);
 	if (Curve)
 	{
diff --git a/Engine/Source/Editor/Persona/Private/SAnimTrackCurvePanel.h b/Engine/Source/Editor/Persona/Private/SAnimTrackCurvePanel.h
--- a/Engine/Source/Editor/Persona/Private/SAnimTrackCurvePanel.h
+++ b/Engine/Source/Editor/Persona/Private/SAnimTrackCurvePanel.h
@@ -91,11 +91,11 @@ private:
 	/**
 	 * Returns true if this curve is editable
 	 */
-	ESlateCheckBoxState::Type	IsCurveEditable(USkeleton::AnimCurveUID Uid) const;
+	ECheckBoxState	IsCurveEditable(USkeleton::AnimCurveUID Uid) const;
 	/**
 	 * Toggle curve visibility
 	 */
-	void		ToggleEditability(ESlateCheckBoxState::Type NewType, USkeleton::AnimCurveUID Uid);
+	void		ToggleEditability(ECheckBoxState NewType, USkeleton::AnimCurveUID Uid);
 	/**
 	 * Refresh Panel
 	 */
@@ -108,10 +108,10 @@ private:
 	/**
 	 * Convert the requested flag bool value into a checkbox state
 	 */
-	ESlateCheckBoxState::Type GetCurveFlagAsCheckboxState(USkeleton::AnimCurveUID CurveUid, EAnimCurveFlags InFlag) const;
+	ECheckBoxState GetCurveFlagAsCheckboxState(USkeleton::AnimCurveUID CurveUid, EAnimCurveFlags InFlag) const;
 
 	/**
 	 * Convert a given checkbox state into a flag value in the provided curve
 	 */
-	void SetCurveFlagFromCheckboxState(ESlateCheckBoxState::Type CheckState, USkeleton::AnimCurveUID CurveUid, EAnimCurveFlags InFlag);
+	void SetCurveFlagFromCheckboxState(ECheckBoxState CheckState, USkeleton::AnimCurveUID CurveUid, EAnimCurveFlags InFlag);
 };
\ No newline at end of file
diff --git a/Engine/Source/Editor/Persona/Private/SAnimationBlendSpace1D.cpp b/Engine/Source/Editor/Persona/Private/SAnimationBlendSpace1D.cpp
--- a/Engine/Source/Editor/Persona/Private/SAnimationBlendSpace1D.cpp
+++ b/Engine/Source/Editor/Persona/Private/SAnimationBlendSpace1D.cpp
@@ -645,24 +645,24 @@ void SBlendSpaceEditor1D::UpdateBlendParameters()
 	BlendSpaceWidget->ResampleData();
 }
 
-ESlateCheckBoxState::Type SBlendSpaceEditor1D::IsEditorDisplayedVertically() const
+ECheckBoxState SBlendSpaceEditor1D::IsEditorDisplayedVertically() const
 {
-	return GetBlendSpace()->bDisplayEditorVertically ? ESlateCheckBoxState::Checked : ESlateCheckBoxState::Unchecked;
+	return GetBlendSpace()->bDisplayEditorVertically ? ECheckBoxState::Checked : ECheckBoxState::Unchecked;
 }
 
-void SBlendSpaceEditor1D::OnChangeDisplayedVertically( ESlateCheckBoxState::Type NewValue )
+void SBlendSpaceEditor1D::OnChangeDisplayedVertically( ECheckBoxState NewValue )
 {
 	bool bPreviousDisplayValue = GetBlendSpace()->bDisplayEditorVertically;
 	switch(NewValue)
 	{
-		case ESlateCheckBoxState::Checked:
+		case ECheckBoxState::Checked:
 		{
 			const FScopedTransaction Transaction( LOCTEXT("Undo_SetDisplayVertically", "Set Blend Space To Display Vertically") );
 			BlendSpace->Modify();
 			GetBlendSpace()->bDisplayEditorVertically = true;
 			break;
 		}
-		case ESlateCheckBoxState::Unchecked:
+		case ECheckBoxState::Unchecked:
 		{
 			const FScopedTransaction Transaction( LOCTEXT("Undo_UnsetDisplayVertically", "Set Blend Space To Not Display Vertically") );
 			BlendSpace->Modify();
diff --git a/Engine/Source/Editor/Persona/Private/SAnimationBlendSpace1D.h b/Engine/Source/Editor/Persona/Private/SAnimationBlendSpace1D.h
--- a/Engine/Source/Editor/Persona/Private/SAnimationBlendSpace1D.h
+++ b/Engine/Source/Editor/Persona/Private/SAnimationBlendSpace1D.h
@@ -186,8 +186,8 @@ private:
 	void OnBlendSpaceParamtersChanged();
 
 	/** Handle display editor vertically checkbox */
-	ESlateCheckBoxState::Type IsEditorDisplayedVertically() const;
-	void OnChangeDisplayedVertically( ESlateCheckBoxState::Type NewValue );
+	ECheckBoxState IsEditorDisplayedVertically() const;
+	void OnChangeDisplayedVertically( ECheckBoxState NewValue );
 
 	/** Get the blend space object we are editing cast to 1D */
 	UBlendSpace1D* GetBlendSpace() {return Cast<UBlendSpace1D>(BlendSpace);}
diff --git a/Engine/Source/Editor/Persona/Private/SAnimationBlendSpaceBase.cpp b/Engine/Source/Editor/Persona/Private/SAnimationBlendSpaceBase.cpp
--- a/Engine/Source/Editor/Persona/Private/SAnimationBlendSpaceBase.cpp
+++ b/Engine/Source/Editor/Persona/Private/SAnimationBlendSpaceBase.cpp
@@ -1194,7 +1194,7 @@ void SBlendSpaceEditorBase::Construct(const FArguments& InArgs)
 */
 void SBlendSpaceEditorBase::Tick( const FGeometry& AllottedGeometry, const double InCurrentTime, const float InDeltaTime )
 {
-	BlendSpaceWidget->bPreviewOn = IsPreviewOn()==ESlateCheckBoxState::Checked;
+	BlendSpaceWidget->bPreviewOn = IsPreviewOn()==ECheckBoxState::Checked;
 	if (BlendSpaceWidget->bPreviewOn)
 	{
 		UpdatePreviewParameter();
@@ -1262,7 +1262,7 @@ TSharedRef<SWidget> SBlendSpaceEditorBase::MakeDisplayOptionsBox() const
 	DisplayOptionsPanel->AddCheckBoxSlot()
 	[
 		SNew(SCheckBox) 
-		.IsChecked( true ? ESlateCheckBoxState::Checked : ESlateCheckBoxState::Unchecked )
+		.IsChecked( true ? ECheckBoxState::Checked : ECheckBoxState::Unchecked )
 		.OnCheckStateChanged( this, &SBlendSpaceEditorBase::ShowToolTip_OnIsCheckedChanged)
 		[
 			SNew(STextBlock)
@@ -1272,7 +1272,7 @@ TSharedRef<SWidget> SBlendSpaceEditorBase::MakeDisplayOptionsBox() const
 	return DisplayOptionsPanel;
 }
 
-ESlateCheckBoxState::Type SBlendSpaceEditorBase::IsPreviewOn() const
+ECheckBoxState SBlendSpaceEditorBase::IsPreviewOn() const
 {
 	class UDebugSkelMeshComponent* Component = PersonaPtr.Pin()->GetPreviewMeshComponent();
 
@@ -1280,15 +1280,15 @@ ESlateCheckBoxState::Type SBlendSpaceEditorBase::IsPreviewOn() const
 	{
 		if (Component->PreviewInstance->CurrentAsset == BlendSpace)
 		{
-			return ESlateCheckBoxState::Checked;
+			return ECheckBoxState::Checked;
 		}
 	}
-	return ESlateCheckBoxState::Unchecked;
+	return ECheckBoxState::Unchecked;
 }
 
-void SBlendSpaceEditorBase::ShowPreview_OnIsCheckedChanged( ESlateCheckBoxState::Type NewValue )
+void SBlendSpaceEditorBase::ShowPreview_OnIsCheckedChanged( ECheckBoxState NewValue )
 {
-	bool bPreviewOn = (NewValue != ESlateCheckBoxState::Unchecked);
+	bool bPreviewOn = (NewValue != ECheckBoxState::Unchecked);
 	class UDebugSkelMeshComponent* Component = PersonaPtr.Pin()->GetPreviewMeshComponent();
 
 	if ( Component != NULL )
@@ -1297,11 +1297,11 @@ void SBlendSpaceEditorBase::ShowPreview_OnIsCheckedChanged( ESlateCheckBoxState:
 	}
 }
 
-void SBlendSpaceEditorBase::ShowToolTip_OnIsCheckedChanged( ESlateCheckBoxState::Type NewValue )
+void SBlendSpaceEditorBase::ShowToolTip_OnIsCheckedChanged( ECheckBoxState NewValue )
 {
 	if (BlendSpaceWidget.IsValid())
 	{
-		BlendSpaceWidget->bTooltipOn = (NewValue != ESlateCheckBoxState::Unchecked);
+		BlendSpaceWidget->bTooltipOn = (NewValue != ECheckBoxState::Unchecked);
 	}
 }
 
diff --git a/Engine/Source/Editor/Persona/Private/SAnimationBlendSpaceBase.h b/Engine/Source/Editor/Persona/Private/SAnimationBlendSpaceBase.h
--- a/Engine/Source/Editor/Persona/Private/SAnimationBlendSpaceBase.h
+++ b/Engine/Source/Editor/Persona/Private/SAnimationBlendSpaceBase.h
@@ -157,7 +157,7 @@ private:
 	FOnParametersUpdated OnParametersUpdated;
 };
 
-typedef TAttribute<ESlateCheckBoxState::Type>::FGetter FReturnCheckboxValueDelegate;
+typedef TAttribute<ECheckBoxState>::FGetter FReturnCheckboxValueDelegate;
 
 struct FBlendSpaceDisplayCheckBoxEntry
 {
@@ -372,9 +372,9 @@ protected:
 	virtual TSharedRef<SWidget> MakeDisplayOptionsBox() const;
 
 	/** Handle preview on blend checkbox */
-	ESlateCheckBoxState::Type IsPreviewOn() const;
-	void ShowPreview_OnIsCheckedChanged( ESlateCheckBoxState::Type NewValue );
-	void ShowToolTip_OnIsCheckedChanged( ESlateCheckBoxState::Type NewValue );
+	ECheckBoxState IsPreviewOn() const;
+	void ShowPreview_OnIsCheckedChanged( ECheckBoxState NewValue );
+	void ShowToolTip_OnIsCheckedChanged( ECheckBoxState NewValue );
 
 	/** Updates Persona's preview window */
 	void UpdatePreviewParameter() const;
diff --git a/Engine/Source/Editor/Persona/Private/SSkeletonTree.cpp b/Engine/Source/Editor/Persona/Private/SSkeletonTree.cpp
--- a/Engine/Source/Editor/Persona/Private/SSkeletonTree.cpp
+++ b/Engine/Source/Editor/Persona/Private/SSkeletonTree.cpp
@@ -741,26 +741,26 @@ TSharedRef< SWidget > FDisplayedAttachedAssetInfo::GenerateWidgetForDataColumn()
 		];
 }
 
-ESlateCheckBoxState::Type FDisplayedAttachedAssetInfo::IsAssetDisplayed() const
+ECheckBoxState FDisplayedAttachedAssetInfo::IsAssetDisplayed() const
 {
 	if(AssetComponent.IsValid())
 	{
-		return AssetComponent->IsVisible() ? ESlateCheckBoxState::Checked : ESlateCheckBoxState::Unchecked;
+		return AssetComponent->IsVisible() ? ECheckBoxState::Checked : ECheckBoxState::Unchecked;
 	}
-	return ESlateCheckBoxState::Undetermined;
+	return ECheckBoxState::Undetermined;
 }
 
-void FDisplayedAttachedAssetInfo::OnToggleAssetDisplayed( ESlateCheckBoxState::Type InCheckboxState )
+void FDisplayedAttachedAssetInfo::OnToggleAssetDisplayed( ECheckBoxState InCheckboxState )
 {
 	if(AssetComponent.IsValid())
 	{
-		AssetComponent->SetVisibility(InCheckboxState == ESlateCheckBoxState::Checked);
+		AssetComponent->SetVisibility(InCheckboxState == ECheckBoxState::Checked);
 	}
 }
 
 const FSlateBrush* FDisplayedAttachedAssetInfo::OnGetAssetDisplayedButtonImage() const
 {
-	return IsAssetDisplayed() == ESlateCheckBoxState::Checked ?
+	return IsAssetDisplayed() == ECheckBoxState::Checked ?
 		FEditorStyle::GetBrush( "Kismet.VariableList.ExposeForInstance" ) :
 		FEditorStyle::GetBrush( "Kismet.VariableList.HideForInstance" );
 }
@@ -2545,15 +2545,15 @@ void SSkeletonTree::AddAttachedAssets( const FPreviewAssetAttachContainer& Attac
 	}
 }
 
-void SSkeletonTree::OnChangeShowingRetargetingOptions(ESlateCheckBoxState::Type NewState)
+void SSkeletonTree::OnChangeShowingRetargetingOptions(ECheckBoxState NewState)
 {
-	bShowingRetargetingOptions = NewState == ESlateCheckBoxState::Checked;
+	bShowingRetargetingOptions = NewState == ECheckBoxState::Checked;
 	CreateTreeColumns();
 }
 
-ESlateCheckBoxState::Type SSkeletonTree::IsShowingRetargetingOptions() const
+ECheckBoxState SSkeletonTree::IsShowingRetargetingOptions() const
 {
-	return bShowingRetargetingOptions ? ESlateCheckBoxState::Checked : ESlateCheckBoxState::Unchecked;
+	return bShowingRetargetingOptions ? ECheckBoxState::Checked : ECheckBoxState::Unchecked;
 }
 
 
diff --git a/Engine/Source/Editor/Persona/Private/SSkeletonTree.h b/Engine/Source/Editor/Persona/Private/SSkeletonTree.h
--- a/Engine/Source/Editor/Persona/Private/SSkeletonTree.h
+++ b/Engine/Source/Editor/Persona/Private/SSkeletonTree.h
@@ -334,10 +334,10 @@ public:
 	UObject* GetAsset() const { return Asset; }
 
 	/** Accessor for SCheckBox **/
-	ESlateCheckBoxState::Type IsAssetDisplayed() const;
+	ECheckBoxState IsAssetDisplayed() const;
 
 	/** Called when user toggles checkbox **/
-	void OnToggleAssetDisplayed( ESlateCheckBoxState::Type InCheckboxState );
+	void OnToggleAssetDisplayed( ECheckBoxState InCheckboxState );
 
 	/** Called when we need to get the state-based-image to show for the asset displayed checkbox */
 	const FSlateBrush* OnGetAssetDisplayedButtonImage() const;
@@ -535,10 +535,10 @@ private:
 	bool IsAddingSocketsAllowed() const;
 
 	/** Handler for "Show Retargeting Options" check box IsChecked functionality */
-	ESlateCheckBoxState::Type IsShowingRetargetingOptions() const;
+	ECheckBoxState IsShowingRetargetingOptions() const;
 
 	/**  Handler for when we change the "Show Retargeting Options" check box */
-	void OnChangeShowingRetargetingOptions(ESlateCheckBoxState::Type NewState);
+	void OnChangeShowingRetargetingOptions(ECheckBoxState NewState);
 
 	/** This replicates the socket filter to the previewcomponent so that the viewport can use the same settings */
 	void SetPreviewComponentSocketFilter() const;
diff --git a/Engine/Source/Editor/Persona/Private/Shared/PersonaMode.cpp b/Engine/Source/Editor/Persona/Private/Shared/PersonaMode.cpp
--- a/Engine/Source/Editor/Persona/Private/Shared/PersonaMode.cpp
+++ b/Engine/Source/Editor/Persona/Private/Shared/PersonaMode.cpp
@@ -373,9 +373,9 @@ EVisibility FAnimBlueprintDefaultsEditorSummoner::IsEditorVisible(EAnimBlueprint
 	return CurrentMode == Mode ? EVisibility::Visible: EVisibility::Hidden;
 }
 
-ESlateCheckBoxState::Type FAnimBlueprintDefaultsEditorSummoner::IsChecked(EAnimBlueprintEditorMode::Type Mode) const
+ECheckBoxState FAnimBlueprintDefaultsEditorSummoner::IsChecked(EAnimBlueprintEditorMode::Type Mode) const
 {
-	return CurrentMode == Mode ? ESlateCheckBoxState::Checked: ESlateCheckBoxState::Unchecked;
+	return CurrentMode == Mode ? ECheckBoxState::Checked: ECheckBoxState::Unchecked;
 }
 
 const FSlateBrush* FAnimBlueprintDefaultsEditorSummoner::GetBorderBrushByMode(EAnimBlueprintEditorMode::Type Mode) const
@@ -390,9 +390,9 @@ const FSlateBrush* FAnimBlueprintDefaultsEditorSummoner::GetBorderBrushByMode(EA
 	}
 }
 
-void FAnimBlueprintDefaultsEditorSummoner::OnCheckedChanged(ESlateCheckBoxState::Type NewType, EAnimBlueprintEditorMode::Type Mode)
+void FAnimBlueprintDefaultsEditorSummoner::OnCheckedChanged(ECheckBoxState NewType, EAnimBlueprintEditorMode::Type Mode)
 {
-	if(NewType == ESlateCheckBoxState::Checked)
+	if(NewType == ECheckBoxState::Checked)
 	{
 		CurrentMode = Mode;
 	}
diff --git a/Engine/Source/Editor/Persona/Private/Shared/PersonaMode.h b/Engine/Source/Editor/Persona/Private/Shared/PersonaMode.h
--- a/Engine/Source/Editor/Persona/Private/Shared/PersonaMode.h
+++ b/Engine/Source/Editor/Persona/Private/Shared/PersonaMode.h
@@ -244,11 +244,11 @@ public:
 private:
 	/** Delegates to customize tab look based on selected mode */
 	EVisibility IsEditorVisible(EAnimBlueprintEditorMode::Type Mode) const;
-	ESlateCheckBoxState::Type IsChecked(EAnimBlueprintEditorMode::Type Mode) const;
+	ECheckBoxState IsChecked(EAnimBlueprintEditorMode::Type Mode) const;
 	const FSlateBrush* GetBorderBrushByMode(EAnimBlueprintEditorMode::Type Mode) const;
 
 	/** Handle changing of editor mode */
-	void OnCheckedChanged(ESlateCheckBoxState::Type NewType, EAnimBlueprintEditorMode::Type Mode);
+	void OnCheckedChanged(ECheckBoxState NewType, EAnimBlueprintEditorMode::Type Mode);
 	EAnimBlueprintEditorMode::Type CurrentMode;
 };
 
diff --git a/Engine/Source/Editor/PhAT/Private/SPhATNewAssetDlg.cpp b/Engine/Source/Editor/PhAT/Private/SPhATNewAssetDlg.cpp
--- a/Engine/Source/Editor/PhAT/Private/SPhATNewAssetDlg.cpp
+++ b/Engine/Source/Editor/PhAT/Private/SPhATNewAssetDlg.cpp
@@ -78,7 +78,7 @@ void SPhATNewAssetDlg::Construct(const FArguments& InArgs)
 				+SUniformGridPanel::Slot(1, 1)
 				[
 					SNew(SCheckBox)
-					.IsChecked(NewBodyData->bAlignDownBone ? ESlateCheckBoxState::Checked : ESlateCheckBoxState::Unchecked)
+					.IsChecked(NewBodyData->bAlignDownBone ? ECheckBoxState::Checked : ECheckBoxState::Unchecked)
 					.OnCheckStateChanged(this, &SPhATNewAssetDlg::OnToggleOrientAlongBone)
 				]
 				+SUniformGridPanel::Slot(0, 2)
@@ -118,7 +118,7 @@ void SPhATNewAssetDlg::Construct(const FArguments& InArgs)
 				+SUniformGridPanel::Slot(1, 4)
 				[
 					SNew(SCheckBox)
-					.IsChecked(NewBodyData->bCreateJoints ? ESlateCheckBoxState::Checked : ESlateCheckBoxState::Unchecked)
+					.IsChecked(NewBodyData->bCreateJoints ? ECheckBoxState::Checked : ECheckBoxState::Unchecked)
 					.OnCheckStateChanged(this, &SPhATNewAssetDlg::OnToggleCreateJoints)
 				]
 				+SUniformGridPanel::Slot(0, 5)
@@ -144,7 +144,7 @@ void SPhATNewAssetDlg::Construct(const FArguments& InArgs)
 				+SUniformGridPanel::Slot(1, 6)
 				[
 					SNew(SCheckBox)
-					.IsChecked(NewBodyData->bWalkPastSmall ? ESlateCheckBoxState::Checked : ESlateCheckBoxState::Unchecked)
+					.IsChecked(NewBodyData->bWalkPastSmall ? ECheckBoxState::Checked : ECheckBoxState::Unchecked)
 					.OnCheckStateChanged(this, &SPhATNewAssetDlg::OnToggleWalkPastSmallBones)
 				]
 				+SUniformGridPanel::Slot(0, 7)
@@ -156,7 +156,7 @@ void SPhATNewAssetDlg::Construct(const FArguments& InArgs)
 				+SUniformGridPanel::Slot(1, 7)
 				[
 					SNew(SCheckBox)
-					.IsChecked(NewBodyData->bBodyForAll ? ESlateCheckBoxState::Checked : ESlateCheckBoxState::Unchecked)
+					.IsChecked(NewBodyData->bBodyForAll ? ECheckBoxState::Checked : ECheckBoxState::Unchecked)
 					.OnCheckStateChanged(this, &SPhATNewAssetDlg::OnToggleCreateBodyForAllBones)
 				]
 				// Add in the UI options for the max hulls and the max verts on the hulls
@@ -312,24 +312,24 @@ void SPhATNewAssetDlg::OnAngularConstraintModeSelectionChanged(TSharedPtr<FStrin
 	}
 }
 
-void SPhATNewAssetDlg::OnToggleOrientAlongBone(ESlateCheckBoxState::Type InCheckboxState)
+void SPhATNewAssetDlg::OnToggleOrientAlongBone(ECheckBoxState InCheckboxState)
 {
-	NewBodyData->bAlignDownBone = (InCheckboxState == ESlateCheckBoxState::Checked);
+	NewBodyData->bAlignDownBone = (InCheckboxState == ECheckBoxState::Checked);
 }
 
-void SPhATNewAssetDlg::OnToggleCreateJoints(ESlateCheckBoxState::Type InCheckboxState)
+void SPhATNewAssetDlg::OnToggleCreateJoints(ECheckBoxState InCheckboxState)
 {
-	NewBodyData->bCreateJoints = (InCheckboxState == ESlateCheckBoxState::Checked);
+	NewBodyData->bCreateJoints = (InCheckboxState == ECheckBoxState::Checked);
 }
 
-void SPhATNewAssetDlg::OnToggleWalkPastSmallBones(ESlateCheckBoxState::Type InCheckboxState)
+void SPhATNewAssetDlg::OnToggleWalkPastSmallBones(ECheckBoxState InCheckboxState)
 {
-	NewBodyData->bWalkPastSmall = (InCheckboxState == ESlateCheckBoxState::Checked);
+	NewBodyData->bWalkPastSmall = (InCheckboxState == ECheckBoxState::Checked);
 }
 
-void SPhATNewAssetDlg::OnToggleCreateBodyForAllBones(ESlateCheckBoxState::Type InCheckboxState)
+void SPhATNewAssetDlg::OnToggleCreateBodyForAllBones(ECheckBoxState InCheckboxState)
 {
-	NewBodyData->bBodyForAll = (InCheckboxState == ESlateCheckBoxState::Checked);
+	NewBodyData->bBodyForAll = (InCheckboxState == ECheckBoxState::Checked);
 }
 
 EVisibility SPhATNewAssetDlg::GetHullOptionsVisibility() const
diff --git a/Engine/Source/Editor/PhAT/Private/SPhATNewAssetDlg.h b/Engine/Source/Editor/PhAT/Private/SPhATNewAssetDlg.h
--- a/Engine/Source/Editor/PhAT/Private/SPhATNewAssetDlg.h
+++ b/Engine/Source/Editor/PhAT/Private/SPhATNewAssetDlg.h
@@ -38,10 +38,10 @@ private:
 	void OnCollisionGeometrySelectionChanged(TSharedPtr<FString> NewSelection, ESelectInfo::Type SelectInfo);
 	void OnWeightOptionSelectionChanged(TSharedPtr<FString> NewSelection, ESelectInfo::Type SelectInfo);
 	void OnAngularConstraintModeSelectionChanged(TSharedPtr<FString> NewSelection, ESelectInfo::Type SelectInfo);
-	void OnToggleOrientAlongBone(ESlateCheckBoxState::Type InCheckboxState);
-	void OnToggleCreateJoints(ESlateCheckBoxState::Type InCheckboxState);
-	void OnToggleWalkPastSmallBones(ESlateCheckBoxState::Type InCheckboxState);
-	void OnToggleCreateBodyForAllBones(ESlateCheckBoxState::Type InCheckboxState);
+	void OnToggleOrientAlongBone(ECheckBoxState InCheckboxState);
+	void OnToggleCreateJoints(ECheckBoxState InCheckboxState);
+	void OnToggleWalkPastSmallBones(ECheckBoxState InCheckboxState);
+	void OnToggleCreateBodyForAllBones(ECheckBoxState InCheckboxState);
 	
 	EVisibility GetHullOptionsVisibility() const;
 	void OnHullCountChanged(int32 InNewValue);	
diff --git a/Engine/Source/Editor/PlacementMode/Private/SPlacementModeTools.cpp b/Engine/Source/Editor/PlacementMode/Private/SPlacementModeTools.cpp
--- a/Engine/Source/Editor/PlacementMode/Private/SPlacementModeTools.cpp
+++ b/Engine/Source/Editor/PlacementMode/Private/SPlacementModeTools.cpp
@@ -542,14 +542,14 @@ TSharedRef< SWidget > SPlacementModeTools::CreatePlacementGroupTab( int32 TabInd
 	];
 }
 
-ESlateCheckBoxState::Type SPlacementModeTools::GetPlacementTabCheckedState( int32 PlacementGroupIndex ) const
+ECheckBoxState SPlacementModeTools::GetPlacementTabCheckedState( int32 PlacementGroupIndex ) const
 {
-	return WidgetSwitcher->GetActiveWidgetIndex() == PlacementGroupIndex ? ESlateCheckBoxState::Checked : ESlateCheckBoxState::Unchecked;
+	return WidgetSwitcher->GetActiveWidgetIndex() == PlacementGroupIndex ? ECheckBoxState::Checked : ECheckBoxState::Unchecked;
 }
 
-void SPlacementModeTools::OnPlacementTabChanged( ESlateCheckBoxState::Type NewState, int32 PlacementGroupIndex )
+void SPlacementModeTools::OnPlacementTabChanged( ECheckBoxState NewState, int32 PlacementGroupIndex )
 {
-	if ( NewState == ESlateCheckBoxState::Checked )
+	if ( NewState == ECheckBoxState::Checked )
 	{
 		WidgetSwitcher->SetActiveWidgetIndex( PlacementGroupIndex );
 
diff --git a/Engine/Source/Editor/PlacementMode/Private/SPlacementModeTools.h b/Engine/Source/Editor/PlacementMode/Private/SPlacementModeTools.h
--- a/Engine/Source/Editor/PlacementMode/Private/SPlacementModeTools.h
+++ b/Engine/Source/Editor/PlacementMode/Private/SPlacementModeTools.h
@@ -66,10 +66,10 @@ private:
 	const FSlateBrush* PlacementGroupBorderImage( int32 PlacementGroupIndex ) const;
 
 	// When the tab is clicked we adjust the check state, so that the right style is displayed.
-	void OnPlacementTabChanged( ESlateCheckBoxState::Type NewState, int32 PlacementGroupIndex );
+	void OnPlacementTabChanged( ECheckBoxState NewState, int32 PlacementGroupIndex );
 
 	// Gets the tab 'active' state, so that we can show the active style
-	ESlateCheckBoxState::Type GetPlacementTabCheckedState( int32 PlacementGroupIndex ) const;
+	ECheckBoxState GetPlacementTabCheckedState( int32 PlacementGroupIndex ) const;
 
 	// Create the standard panel displayed when no search is being performed.
 	TSharedRef< SWidget > CreateStandardPanel();
diff --git a/Engine/Source/Editor/ProjectTargetPlatformEditor/Private/Widgets/SProjectTargetPlatformSettings.cpp b/Engine/Source/Editor/ProjectTargetPlatformEditor/Private/Widgets/SProjectTargetPlatformSettings.cpp
--- a/Engine/Source/Editor/ProjectTargetPlatformEditor/Private/Widgets/SProjectTargetPlatformSettings.cpp
+++ b/Engine/Source/Editor/ProjectTargetPlatformEditor/Private/Widgets/SProjectTargetPlatformSettings.cpp
@@ -123,7 +123,7 @@ TSharedRef<SWidget> SProjectTargetPlatformSettings::MakePlatformRow(const FText&
 		];
 }
 
-ESlateCheckBoxState::Type SProjectTargetPlatformSettings::HandlePlatformCheckBoxIsChecked(const FName PlatformName) const
+ECheckBoxState SProjectTargetPlatformSettings::HandlePlatformCheckBoxIsChecked(const FName PlatformName) const
 {
 	FProjectStatus ProjectStatus;
 	if(IProjectManager::Get().QueryStatusForCurrentProject(ProjectStatus))
@@ -131,15 +131,15 @@ ESlateCheckBoxState::Type SProjectTargetPlatformSettings::HandlePlatformCheckBox
 		if(PlatformName.IsNone())
 		{
 			// None is "All Platforms"
-			return (ProjectStatus.SupportsAllPlatforms()) ? ESlateCheckBoxState::Checked : ESlateCheckBoxState::Unchecked;
+			return (ProjectStatus.SupportsAllPlatforms()) ? ECheckBoxState::Checked : ECheckBoxState::Unchecked;
 		}
 		else
 		{
-			return (ProjectStatus.IsTargetPlatformSupported(PlatformName)) ? ESlateCheckBoxState::Checked : ESlateCheckBoxState::Unchecked;
+			return (ProjectStatus.IsTargetPlatformSupported(PlatformName)) ? ECheckBoxState::Checked : ECheckBoxState::Unchecked;
 		}
 	}
 
-	return ESlateCheckBoxState::Unchecked;
+	return ECheckBoxState::Unchecked;
 }
 
 bool SProjectTargetPlatformSettings::HandlePlatformCheckBoxIsEnabled(const FName PlatformName) const
@@ -162,12 +162,12 @@ bool SProjectTargetPlatformSettings::HandlePlatformCheckBoxIsEnabled(const FName
 	return false;
 }
 
-void SProjectTargetPlatformSettings::HandlePlatformCheckBoxStateChanged(ESlateCheckBoxState::Type InState, const FName PlatformName) const
+void SProjectTargetPlatformSettings::HandlePlatformCheckBoxStateChanged(ECheckBoxState InState, const FName PlatformName) const
 {
 	if(PlatformName.IsNone())
 	{
 		// None is "All Platforms"
-		if(InState == ESlateCheckBoxState::Checked)
+		if(InState == ECheckBoxState::Checked)
 		{
 			FGameProjectGenerationModule::Get().ClearSupportedTargetPlatforms();
 		}
@@ -182,7 +182,7 @@ void SProjectTargetPlatformSettings::HandlePlatformCheckBoxStateChanged(ESlateCh
 	}
 	else
 	{
-		FGameProjectGenerationModule::Get().UpdateSupportedTargetPlatforms(PlatformName, InState == ESlateCheckBoxState::Checked);
+		FGameProjectGenerationModule::Get().UpdateSupportedTargetPlatforms(PlatformName, InState == ECheckBoxState::Checked);
 	}
 }
 
diff --git a/Engine/Source/Editor/ProjectTargetPlatformEditor/Private/Widgets/SProjectTargetPlatformSettings.h b/Engine/Source/Editor/ProjectTargetPlatformEditor/Private/Widgets/SProjectTargetPlatformSettings.h
--- a/Engine/Source/Editor/ProjectTargetPlatformEditor/Private/Widgets/SProjectTargetPlatformSettings.h
+++ b/Engine/Source/Editor/ProjectTargetPlatformEditor/Private/Widgets/SProjectTargetPlatformSettings.h
@@ -17,13 +17,13 @@ private:
 	TSharedRef<SWidget> MakePlatformRow(const FText& DisplayName, const FName PlatformName, const FSlateBrush* Icon) const;
 
 	/** Check to see if the "enabled" checkbox should be checked for this platform */
-	ESlateCheckBoxState::Type HandlePlatformCheckBoxIsChecked(const FName PlatformName) const;
+	ECheckBoxState HandlePlatformCheckBoxIsChecked(const FName PlatformName) const;
 
 	/** Check to see if the "enabled" checkbox should be enabled for this platform */
 	bool HandlePlatformCheckBoxIsEnabled(const FName PlatformName) const;
 
 	/** Handle the "enabled" checkbox state being changed for this platform */
-	void HandlePlatformCheckBoxStateChanged(ESlateCheckBoxState::Type InState, const FName PlatformName) const;
+	void HandlePlatformCheckBoxStateChanged(ECheckBoxState InState, const FName PlatformName) const;
 
 	TArray<const PlatformInfo::FPlatformInfo*> AvailablePlatforms;
 };
diff --git a/Engine/Source/Editor/PropertyEditor/Private/PropertyEditorHelpers.cpp b/Engine/Source/Editor/PropertyEditor/Private/PropertyEditorHelpers.cpp
--- a/Engine/Source/Editor/PropertyEditor/Private/PropertyEditorHelpers.cpp
+++ b/Engine/Source/Editor/PropertyEditor/Private/PropertyEditorHelpers.cpp
@@ -274,22 +274,22 @@ bool SEditConditionWidget::HasEditCondition() const
 		||	( CustomEditCondition.OnEditConditionValueChanged.IsBound() );
 }
 
-void SEditConditionWidget::OnEditConditionCheckChanged( ESlateCheckBoxState::Type CheckState )
+void SEditConditionWidget::OnEditConditionCheckChanged( ECheckBoxState CheckState )
 {
 	if( PropertyEditor.IsValid() && PropertyEditor->HasEditCondition() && PropertyEditor->SupportsEditConditionToggle() )
 	{
-		PropertyEditor->SetEditConditionState( CheckState == ESlateCheckBoxState::Checked );
+		PropertyEditor->SetEditConditionState( CheckState == ECheckBoxState::Checked );
 	}
 	else
 	{
-		CustomEditCondition.OnEditConditionValueChanged.ExecuteIfBound( CheckState == ESlateCheckBoxState::Checked );
+		CustomEditCondition.OnEditConditionValueChanged.ExecuteIfBound( CheckState == ECheckBoxState::Checked );
 	}
 }
 
-ESlateCheckBoxState::Type SEditConditionWidget::OnGetEditConditionCheckState() const
+ECheckBoxState SEditConditionWidget::OnGetEditConditionCheckState() const
 {
 	bool bEditConditionMet = ( PropertyEditor.IsValid() && PropertyEditor->HasEditCondition() && PropertyEditor->IsEditConditionMet() ) || CustomEditCondition.EditConditionValue.Get();
-	return bEditConditionMet ? ESlateCheckBoxState::Checked : ESlateCheckBoxState::Unchecked;
+	return bEditConditionMet ? ECheckBoxState::Checked : ECheckBoxState::Unchecked;
 }
 
 namespace PropertyEditorHelpers
diff --git a/Engine/Source/Editor/PropertyEditor/Private/PropertyEditorHelpers.h b/Engine/Source/Editor/PropertyEditor/Private/PropertyEditorHelpers.h
--- a/Engine/Source/Editor/PropertyEditor/Private/PropertyEditorHelpers.h
+++ b/Engine/Source/Editor/PropertyEditor/Private/PropertyEditorHelpers.h
@@ -88,9 +88,9 @@ class SEditConditionWidget : public SCompoundWidget
 	void Construct( const FArguments& Args, TSharedPtr<FPropertyEditor> InPropertyEditor );
 
 private:
-	void OnEditConditionCheckChanged( ESlateCheckBoxState::Type CheckState );
+	void OnEditConditionCheckChanged( ECheckBoxState CheckState );
 	bool HasEditCondition() const;
-	ESlateCheckBoxState::Type OnGetEditConditionCheckState() const;
+	ECheckBoxState OnGetEditConditionCheckState() const;
 private:
 	TSharedPtr<FPropertyEditor> PropertyEditor;
 	FCustomEditCondition CustomEditCondition;
diff --git a/Engine/Source/Editor/PropertyEditor/Private/UserInterface/PropertyEditor/SPropertyEditorBool.cpp b/Engine/Source/Editor/PropertyEditor/Private/UserInterface/PropertyEditor/SPropertyEditorBool.cpp
--- a/Engine/Source/Editor/PropertyEditor/Private/UserInterface/PropertyEditor/SPropertyEditorBool.cpp
+++ b/Engine/Source/Editor/PropertyEditor/Private/UserInterface/PropertyEditor/SPropertyEditorBool.cpp
@@ -53,9 +53,9 @@ FReply SPropertyEditorBool::OnFocusReceived( const FGeometry& MyGeometry, const
 	return FReply::Handled().SetUserFocus(CheckBox.ToSharedRef(), InFocusEvent.GetCause());
 }
 
-ESlateCheckBoxState::Type SPropertyEditorBool::OnGetCheckState() const
+ECheckBoxState SPropertyEditorBool::OnGetCheckState() const
 {
-	ESlateCheckBoxState::Type ReturnState = ESlateCheckBoxState::Undetermined;
+	ECheckBoxState ReturnState = ECheckBoxState::Undetermined;
 
 	bool Value;
 	const TSharedRef< IPropertyHandle > PropertyHandle = PropertyEditor->GetPropertyHandle();
@@ -63,21 +63,21 @@ ESlateCheckBoxState::Type SPropertyEditorBool::OnGetCheckState() const
 	{
 		if( Value == true )
 		{
-			ReturnState = ESlateCheckBoxState::Checked;
+			ReturnState = ECheckBoxState::Checked;
 		}
 		else if( Value == false )
 		{
-			ReturnState = ESlateCheckBoxState::Unchecked;
+			ReturnState = ECheckBoxState::Unchecked;
 		}
 	}
 
 	return ReturnState;
 }
 
-void SPropertyEditorBool::OnCheckStateChanged( ESlateCheckBoxState::Type InNewState )
+void SPropertyEditorBool::OnCheckStateChanged( ECheckBoxState InNewState )
 {
 	const TSharedRef< IPropertyHandle > PropertyHandle = PropertyEditor->GetPropertyHandle();
-	if( InNewState == ESlateCheckBoxState::Checked || InNewState == ESlateCheckBoxState::Undetermined )
+	if( InNewState == ECheckBoxState::Checked || InNewState == ECheckBoxState::Undetermined )
 	{
 		PropertyHandle->SetValue( true );
 	}
diff --git a/Engine/Source/Editor/PropertyEditor/Private/UserInterface/PropertyEditor/SPropertyEditorBool.h b/Engine/Source/Editor/PropertyEditor/Private/UserInterface/PropertyEditor/SPropertyEditorBool.h
--- a/Engine/Source/Editor/PropertyEditor/Private/UserInterface/PropertyEditor/SPropertyEditorBool.h
+++ b/Engine/Source/Editor/PropertyEditor/Private/UserInterface/PropertyEditor/SPropertyEditorBool.h
@@ -23,9 +23,9 @@ private:
 
 	virtual FReply OnMouseButtonDoubleClick( const FGeometry& InMyGeometry, const FPointerEvent& InMouseEvent );
 
-	ESlateCheckBoxState::Type OnGetCheckState() const;
+	ECheckBoxState OnGetCheckState() const;
 
-	void OnCheckStateChanged( ESlateCheckBoxState::Type InNewState );
+	void OnCheckStateChanged( ECheckBoxState InNewState );
 
 	/** @return True if the property can be edited */
 	bool CanEdit() const;
diff --git a/Engine/Source/Editor/PropertyEditor/Private/UserInterface/PropertyEditor/SPropertyEditorTableRow.cpp b/Engine/Source/Editor/PropertyEditor/Private/UserInterface/PropertyEditor/SPropertyEditorTableRow.cpp
--- a/Engine/Source/Editor/PropertyEditor/Private/UserInterface/PropertyEditor/SPropertyEditorTableRow.cpp
+++ b/Engine/Source/Editor/PropertyEditor/Private/UserInterface/PropertyEditor/SPropertyEditorTableRow.cpp
@@ -193,15 +193,15 @@ const FSlateBrush* SPropertyEditorTableRow::OnGetFavoriteImage() const
 	return FEditorStyle::GetBrush(TEXT("PropertyWindow.Favorites_Disabled"));
 }
 
-void SPropertyEditorTableRow::OnEditConditionCheckChanged( ESlateCheckBoxState::Type CheckState )
+void SPropertyEditorTableRow::OnEditConditionCheckChanged( ECheckBoxState CheckState )
 {
-	PropertyEditor->SetEditConditionState( CheckState == ESlateCheckBoxState::Checked );
+	PropertyEditor->SetEditConditionState( CheckState == ECheckBoxState::Checked );
 }
 
-ESlateCheckBoxState::Type SPropertyEditorTableRow::OnGetEditConditionCheckState() const
+ECheckBoxState SPropertyEditorTableRow::OnGetEditConditionCheckState() const
 {
 	bool bEditConditionMet = PropertyEditor->IsEditConditionMet();
-	return bEditConditionMet ? ESlateCheckBoxState::Checked : ESlateCheckBoxState::Unchecked;
+	return bEditConditionMet ? ECheckBoxState::Checked : ECheckBoxState::Unchecked;
 }
 
 FReply SPropertyEditorTableRow::OnNameDoubleClicked()
diff --git a/Engine/Source/Editor/PropertyEditor/Private/UserInterface/PropertyEditor/SPropertyEditorTableRow.h b/Engine/Source/Editor/PropertyEditor/Private/UserInterface/PropertyEditor/SPropertyEditorTableRow.h
--- a/Engine/Source/Editor/PropertyEditor/Private/UserInterface/PropertyEditor/SPropertyEditorTableRow.h
+++ b/Engine/Source/Editor/PropertyEditor/Private/UserInterface/PropertyEditor/SPropertyEditorTableRow.h
@@ -42,8 +42,8 @@ private:
 	const FSlateBrush* OnGetFavoriteImage() const;
 	FReply OnToggleFavoriteClicked();
 
-	void OnEditConditionCheckChanged( ESlateCheckBoxState::Type CheckState );
-	ESlateCheckBoxState::Type OnGetEditConditionCheckState() const;
+	void OnEditConditionCheckChanged( ECheckBoxState CheckState );
+	ECheckBoxState OnGetEditConditionCheckState() const;
 
 	FReply OnNameDoubleClicked();
 
diff --git a/Engine/Source/Editor/ReferenceViewer/Private/SReferenceViewer.cpp b/Engine/Source/Editor/ReferenceViewer/Private/SReferenceViewer.cpp
--- a/Engine/Source/Editor/ReferenceViewer/Private/SReferenceViewer.cpp
+++ b/Engine/Source/Editor/ReferenceViewer/Private/SReferenceViewer.cpp
@@ -387,24 +387,24 @@ void SReferenceViewer::OnUpdateHistoryData(FReferenceViewerHistoryData& HistoryD
 	}
 }
 
-void SReferenceViewer::OnSearchDepthEnabledChanged( ESlateCheckBoxState::Type NewState )
+void SReferenceViewer::OnSearchDepthEnabledChanged( ECheckBoxState NewState )
 {
 	if ( GraphObj )
 	{
-		GraphObj->SetSearchDepthLimitEnabled(NewState == ESlateCheckBoxState::Checked);
+		GraphObj->SetSearchDepthLimitEnabled(NewState == ECheckBoxState::Checked);
 		GraphObj->RebuildGraph();
 	}
 }
 
-ESlateCheckBoxState::Type SReferenceViewer::IsSearchDepthEnabledChecked() const
+ECheckBoxState SReferenceViewer::IsSearchDepthEnabledChecked() const
 {
 	if ( GraphObj )
 	{
-		return GraphObj->IsSearchDepthLimited() ? ESlateCheckBoxState::Checked : ESlateCheckBoxState::Unchecked;
+		return GraphObj->IsSearchDepthLimited() ? ECheckBoxState::Checked : ECheckBoxState::Unchecked;
 	}
 	else
 	{
-		return ESlateCheckBoxState::Unchecked;
+		return ECheckBoxState::Unchecked;
 	}
 }
 
@@ -429,24 +429,24 @@ void SReferenceViewer::OnSearchDepthCommitted(int32 NewValue)
 	}
 }
 
-void SReferenceViewer::OnSearchBreadthEnabledChanged( ESlateCheckBoxState::Type NewState )
+void SReferenceViewer::OnSearchBreadthEnabledChanged( ECheckBoxState NewState )
 {
 	if ( GraphObj )
 	{
-		GraphObj->SetSearchBreadthLimitEnabled(NewState == ESlateCheckBoxState::Checked);
+		GraphObj->SetSearchBreadthLimitEnabled(NewState == ECheckBoxState::Checked);
 		GraphObj->RebuildGraph();
 	}
 }
 
-ESlateCheckBoxState::Type SReferenceViewer::IsSearchBreadthEnabledChecked() const
+ECheckBoxState SReferenceViewer::IsSearchBreadthEnabledChecked() const
 {
 	if ( GraphObj )
 	{
-		return GraphObj->IsSearchBreadthLimited() ? ESlateCheckBoxState::Checked : ESlateCheckBoxState::Unchecked;
+		return GraphObj->IsSearchBreadthLimited() ? ECheckBoxState::Checked : ECheckBoxState::Unchecked;
 	}
 	else
 	{
-		return ESlateCheckBoxState::Unchecked;
+		return ECheckBoxState::Unchecked;
 	}
 }
 
