Editor: Remove a bunch of GEditor dependencies in mesh painting code

diff --git a/Engine/Source/Editor/MeshPaint/Private/MeshPaintEdMode.cpp b/Engine/Source/Editor/MeshPaint/Private/MeshPaintEdMode.cpp
--- a/Engine/Source/Editor/MeshPaint/Private/MeshPaintEdMode.cpp
+++ b/Engine/Source/Editor/MeshPaint/Private/MeshPaintEdMode.cpp
@@ -151,7 +151,7 @@ void FEdModeMeshPaint::Enter()
 	
 	if (!Toolkit.IsValid())
 	{
-		Toolkit = MakeShareable(new FMeshPaintToolKit);
+		Toolkit = MakeShareable(new FMeshPaintToolKit(this));
 		Toolkit->Init(Owner->GetToolkitHost());
 	}
 
@@ -460,7 +460,7 @@ bool FEdModeMeshPaint::InputKey( FEditorViewportClient* InViewportClient, FViewp
 			if (HitProxy && HitProxy->IsA(HActor::StaticGetType()))
 			{
 				const AActor* ClickedActor = (static_cast<const HActor*>(HitProxy))->Actor;
-				USelection& SelectedActors = *GEditor->GetSelectedActors();
+				USelection& SelectedActors = *Owner->GetSelectedActors();
 				if (SelectedActors.IsSelected(ClickedActor))
 				{
 					// Clicked actor is currently selected, start painting.
@@ -935,7 +935,7 @@ void FEdModeMeshPaint::DoPaint( const FVector& InCameraOrigin,
 		const FVector TraceEnd( InRayOrigin + InRayDirection * HALF_WORLD_MAX );
 
 		// Iterate over selected actors looking for static meshes
-		USelection& SelectedActors = *GEditor->GetSelectedActors();
+		USelection& SelectedActors = *Owner->GetSelectedActors();
 		TArray< AActor* > ValidSelectedActors;
 		for( int32 CurSelectedActorIndex = 0; CurSelectedActorIndex < SelectedActors.Num(); ++CurSelectedActorIndex )
 		{
@@ -2785,16 +2785,14 @@ void FEdModeMeshPaint::Render( const FSceneView* View, FViewport* Viewport, FPri
 	}
 }
 
-
-
 // @TODO MeshPaint: Cache selected static mesh components each time selection change
 /** Returns valid StaticMesheComponents in the current selection */
-TArray<UStaticMeshComponent*> GetValidStaticMeshComponents()
+TArray<UStaticMeshComponent*> FEdModeMeshPaint::GetValidStaticMeshComponents() const
 {
 	TArray<UStaticMeshComponent*> SMComponents;
 
 	// Iterate over selected actors looking for static meshes
-	USelection& SelectedActors = *GEditor->GetSelectedActors();
+	USelection& SelectedActors = *Owner->GetSelectedActors();
 	for( int32 CurSelectedActorIndex = 0; CurSelectedActorIndex < SelectedActors.Num(); ++CurSelectedActorIndex )
 	{
 		AActor* CurActor = CastChecked< AActor >( SelectedActors.GetSelectedObject( CurSelectedActorIndex ) );
@@ -3621,7 +3619,7 @@ void FEdModeMeshPaint::RemoveInstanceVertexColors() const
 {
 	FScopedTransaction Transaction( LOCTEXT( "MeshPaintMode_VertexPaint_TransactionRemoveInstColors", "Remove Instance Vertex Colors" ) );
 
-	USelection& SelectedActors = *GEditor->GetSelectedActors();
+	USelection& SelectedActors = *Owner->GetSelectedActors();
 	for( int32 CurSelectedActorIndex = 0; CurSelectedActorIndex < SelectedActors.Num(); ++CurSelectedActorIndex )
 	{
 		RemoveInstanceVertexColors( SelectedActors.GetSelectedObject( CurSelectedActorIndex ) );
@@ -3654,7 +3652,7 @@ void FEdModeMeshPaint::CopyInstanceVertexColors()
 {
 	CopiedColorsByComponent.Empty();
 
-	USelection& SelectedActors = *GEditor->GetSelectedActors();
+	USelection& SelectedActors = *Owner->GetSelectedActors();
 	if( SelectedActors.Num() != 1 )
 	{
 		// warning - works only with 1 actor selected..!
@@ -3736,7 +3734,7 @@ void FEdModeMeshPaint::PasteInstanceVertexColors()
 
 	FScopedTransaction Transaction( LOCTEXT( "MeshPaintMode_VertexPaint_TransactionPasteInstColors", "Paste Instance Vertex Colors" ) );
 
-	USelection& SelectedActors = *GEditor->GetSelectedActors();
+	USelection& SelectedActors = *Owner->GetSelectedActors();
 
 	TScopedPointer< FComponentReregisterContext > ComponentReregisterContext;
 
@@ -3861,7 +3859,7 @@ bool FEdModeMeshPaint::RequiresInstanceVertexColorsFixup() const
 	bool bRequiresFixup = false;
 
 	// Find each static mesh component of any selected actors
-	USelection& SelectedActors = *GEditor->GetSelectedActors();
+	USelection& SelectedActors = *Owner->GetSelectedActors();
 	for( int32 CurSelectedActorIndex = 0; CurSelectedActorIndex < SelectedActors.Num(); ++CurSelectedActorIndex )
 	{
 		AActor* SelectedActor = Cast< AActor >( SelectedActors.GetSelectedObject( CurSelectedActorIndex ) );
@@ -3889,7 +3887,7 @@ bool FEdModeMeshPaint::RequiresInstanceVertexColorsFixup() const
 void FEdModeMeshPaint::FixupInstanceVertexColors() const
 {
 	// Find each static mesh component of any selected actors
-	USelection& SelectedActors = *GEditor->GetSelectedActors();
+	USelection& SelectedActors = *Owner->GetSelectedActors();
 	for( int32 CurSelectedActorIndex = 0; CurSelectedActorIndex < SelectedActors.Num(); ++CurSelectedActorIndex )
 	{
 		AActor* SelectedActor = Cast< AActor >( SelectedActors.GetSelectedObject( CurSelectedActorIndex ) );
@@ -3911,7 +3909,7 @@ void FEdModeMeshPaint::FixupInstanceVertexColors() const
 
 void FEdModeMeshPaint::ForceBestLOD()
 {
-	USelection& SelectedActors = *GEditor->GetSelectedActors();
+	USelection& SelectedActors = *Owner->GetSelectedActors();
 	for( int32 CurSelectedActorIndex = 0; CurSelectedActorIndex < SelectedActors.Num(); ++CurSelectedActorIndex )
 	{
 		AActor* SelectedActor = Cast< AActor >( SelectedActors.GetSelectedObject( CurSelectedActorIndex ) );
@@ -3939,7 +3937,7 @@ void FEdModeMeshPaint::ForceBestLOD(UStaticMeshComponent* StaticMeshComponent)
 
 void FEdModeMeshPaint::ClearForcedLOD()
 {
-	USelection& SelectedActors = *GEditor->GetSelectedActors();
+	USelection& SelectedActors = *Owner->GetSelectedActors();
 	for( int32 CurSelectedActorIndex = 0; CurSelectedActorIndex < SelectedActors.Num(); ++CurSelectedActorIndex )
 	{
 		AActor* SelectedActor = Cast< AActor >( SelectedActors.GetSelectedObject( CurSelectedActorIndex ) );
@@ -3968,7 +3966,7 @@ void FEdModeMeshPaint::ClearForcedLOD(UStaticMeshComponent* StaticMeshComponent)
 void FEdModeMeshPaint::ApplyVertexColorsToAllLODs()
 {	
 	// Find each static mesh component of any selected actors
-	USelection& SelectedActors = *GEditor->GetSelectedActors();
+	USelection& SelectedActors = *Owner->GetSelectedActors();
 	for( int32 CurSelectedActorIndex = 0; CurSelectedActorIndex < SelectedActors.Num(); ++CurSelectedActorIndex )
 	{
 		AActor* SelectedActor = Cast< AActor >( SelectedActors.GetSelectedObject( CurSelectedActorIndex ) );
@@ -4100,7 +4098,7 @@ void FEdModeMeshPaint::CreateInstanceMaterialAndTexture() const
 /** Removes instance of paintable material/texture for the selected mesh */
 void FEdModeMeshPaint::RemoveInstanceMaterialAndTexture() const
 {
-	USelection& SelectedActors = *GEditor->GetSelectedActors();
+	USelection& SelectedActors = *Owner->GetSelectedActors();
 	for( int32 CurSelectedActorIndex = 0; CurSelectedActorIndex < SelectedActors.Num(); ++CurSelectedActorIndex )
 	{
 		AActor* SelectedActor = Cast< AActor >( SelectedActors.GetSelectedObject( CurSelectedActorIndex ) );
@@ -4130,7 +4128,7 @@ bool FEdModeMeshPaint::GetSelectedMeshInfo( int32& OutTotalBaseVertexColorBytes,
 
 	int32 NumValidMeshes = 0;
 
-	USelection& SelectedActors = *GEditor->GetSelectedActors();
+	USelection& SelectedActors = *Owner->GetSelectedActors();
 	for( int32 CurSelectedActorIndex = 0; CurSelectedActorIndex < SelectedActors.Num(); ++CurSelectedActorIndex )
 	{
 		AActor* SelectedActor = Cast< AActor >( SelectedActors.GetSelectedObject( CurSelectedActorIndex ) );
@@ -4255,7 +4253,7 @@ void FImportVertexTextureHelper::PickVertexColorFromTex(FColor& NewVertexColor,
 }
 
 
-void FImportVertexTextureHelper::ImportVertexColors(const FString& Filename, int32 UVIndex, int32 ImportLOD, uint8 ColorMask)
+void FImportVertexTextureHelper::ImportVertexColors(FEditorModeTools* ModeTools, const FString& Filename, int32 UVIndex, int32 ImportLOD, uint8 ColorMask)
 {
 	FMessageLog EditorErrors("EditorErrors");
 	EditorErrors.NewPage(LOCTEXT("MeshPaintImportLogLabel", "Mesh Paint: Import Vertex Colors"));
@@ -4268,7 +4266,7 @@ void FImportVertexTextureHelper::ImportVertexColors(const FString& Filename, int
 	}
 
 	TArray<UStaticMeshComponent*> Components;
-	for ( FSelectionIterator It( GEditor->GetSelectedActorIterator() ) ; It ; ++It )
+	for ( FSelectionIterator It( *ModeTools->GetSelectedActors() ) ; It ; ++It )
 	{
 		AActor* Actor = Cast<AActor>( *It );
 		if(Actor)
diff --git a/Engine/Source/Editor/MeshPaint/Private/SMeshPaint.cpp b/Engine/Source/Editor/MeshPaint/Private/SMeshPaint.cpp
--- a/Engine/Source/Editor/MeshPaint/Private/SMeshPaint.cpp
+++ b/Engine/Source/Editor/MeshPaint/Private/SMeshPaint.cpp
@@ -34,8 +34,10 @@ public:
 	}
 
 	BEGIN_SLATE_FUNCTION_BUILD_OPTIMIZATION
-	void Construct( const FArguments& InArgs )
+	void Construct( const FArguments& InArgs, class FEdModeMeshPaint* InMeshPaintEditMode )
 	{	
+		MeshPaintEditMode = InMeshPaintEditMode;
+			
 		ChildSlot
 		[
 			SNew(SBorder)
@@ -288,7 +290,7 @@ private:
 	{
 		FImportVertexTextureHelper ImportVertex;
 		FString Path = TGATextBox->GetText().ToString();
-		ImportVertex.ImportVertexColors(Path, UVValue, LODValue, ImportColorMask);
+		ImportVertex.ImportVertexColors(MeshPaintEditMode->GetModeManager(), Path, UVValue, LODValue, ImportColorMask);
 
 		return FReply::Handled();
 	}
@@ -351,8 +353,11 @@ private:
 	/** The currently selected LOD value */
 	int32 LODValue;
 
-	/* Mask representing the import color channels that are selected */
+	/** Mask representing the import color channels that are selected */
 	uint8 ImportColorMask;
+	
+	/** Mesh paint edit mode */
+	class FEdModeMeshPaint* MeshPaintEditMode;
 };
 
 DECLARE_DELEGATE_OneParam( FOMeshPaintResourceChanged, EMeshPaintResource::Type );
@@ -1063,6 +1068,10 @@ private:
 };
 
 
+FMeshPaintToolKit::FMeshPaintToolKit(class FEdModeMeshPaint* InOwningMode)
+	: MeshPaintEdMode(InOwningMode)
+{
+}
 
 void FMeshPaintToolKit::RegisterTabSpawners(const TSharedRef<class FTabManager>& TabManager)
 {
@@ -1076,7 +1085,7 @@ void FMeshPaintToolKit::UnregisterTabSpawners(const TSharedRef<class FTabManager
 
 void FMeshPaintToolKit::Init(const TSharedPtr< class IToolkitHost >& InitToolkitHost)
 {
-	MeshPaintWidgets = SNew(SMeshPaint, SharedThis(this));
+	MeshPaintWidgets = SNew(SMeshPaint, SharedThis(this), MeshPaintEdMode);
 
 	FModeToolkit::Init(InitToolkitHost);
 }
@@ -1093,7 +1102,7 @@ FText FMeshPaintToolKit::GetBaseToolkitName() const
 
 class FEdMode* FMeshPaintToolKit::GetEditorMode() const
 {
-	return GLevelEditorModeTools().GetActiveMode(FBuiltinEditorModes::EM_MeshPaint);
+	return MeshPaintEdMode;
 }
 
 TSharedPtr<SWidget> FMeshPaintToolKit::GetInlineContent() const
@@ -1107,9 +1116,9 @@ SMeshPaint::~SMeshPaint()
 }
 
 BEGIN_SLATE_FUNCTION_BUILD_OPTIMIZATION
-void SMeshPaint::Construct(const FArguments& InArgs, TSharedRef<FMeshPaintToolKit> InParentToolkit)
+void SMeshPaint::Construct(const FArguments& InArgs, TSharedRef<FMeshPaintToolKit> InParentToolkit, class FEdModeMeshPaint* InOwningMode)
 {
-	MeshPaintEditMode = (FEdModeMeshPaint*)GLevelEditorModeTools().GetActiveMode( FBuiltinEditorModes::EM_MeshPaint );
+	MeshPaintEditMode = InOwningMode;
 
 	PaintColorSet = EMeshPaintColorSet::PaintColor;
 	bShowImportOptions = false;
@@ -1672,7 +1681,7 @@ void SMeshPaint::Construct(const FArguments& InArgs, TSharedRef<FMeshPaintToolKi
 								+SVerticalBox::Slot()
 								.AutoHeight()
 								[
-									SNew(SImportVertexColorsFromTGA)
+									SNew(SImportVertexColorsFromTGA, MeshPaintEditMode)
 								]
 							]
 							// Erase/Paint color radio button
@@ -2093,9 +2102,9 @@ void SMeshPaint::Tick( const FGeometry& AllottedGeometry, const double InCurrent
 	MeshPaintEditMode->UpdateTexturePaintTargetList();
 }
 
-class FEdMode* SMeshPaint::GetEditorMode() const
+class FEdModeMeshPaint* SMeshPaint::GetEditorMode() const
 {
-	return (FEdModeMeshPaint*)GLevelEditorModeTools().GetActiveMode(FBuiltinEditorModes::EM_MeshPaint);
+	return MeshPaintEditMode;
 }
 
 EVisibility SMeshPaint::GetResourceTypeVerticesVisibility() const
diff --git a/Engine/Source/Editor/MeshPaint/Private/SMeshPaint.h b/Engine/Source/Editor/MeshPaint/Private/SMeshPaint.h
--- a/Engine/Source/Editor/MeshPaint/Private/SMeshPaint.h
+++ b/Engine/Source/Editor/MeshPaint/Private/SMeshPaint.h
@@ -39,6 +39,8 @@ namespace EMeshPaintWriteColorChannels
 class FMeshPaintToolKit : public FModeToolkit
 {
 public:
+	FMeshPaintToolKit(class FEdModeMeshPaint* InOwningMode);
+
 	virtual void RegisterTabSpawners(const TSharedRef<class FTabManager>& TabManager) override;
 	virtual void UnregisterTabSpawners(const TSharedRef<class FTabManager>& TabManager) override;
 
@@ -54,6 +56,9 @@ public:
 private:
 	/** Geometry tools widget */
 	TSharedPtr<class SMeshPaint> MeshPaintWidgets;
+
+	/** Owning editor mode */
+	class FEdModeMeshPaint* MeshPaintEdMode;
 };
 
 /**
@@ -68,14 +73,14 @@ public:
 	~SMeshPaint();
 
 	/** SCompoundWidget functions */
-	void Construct(const FArguments& InArgs, TSharedRef<FMeshPaintToolKit> InParentToolkit);
+	void Construct(const FArguments& InArgs, TSharedRef<FMeshPaintToolKit> InParentToolkit, class FEdModeMeshPaint* InOwningMode);
 
 	virtual void Tick( const FGeometry& AllottedGeometry, const double InCurrentTime, const float InDeltaTime ) override;
 
 protected:
 
 	/** Gets a reference to the mesh paint editor mode if it is active */
-	FEdMode* GetEditorMode() const;
+	class FEdModeMeshPaint* GetEditorMode() const;
 
 	/** Returns the visibility state of vertices only property controls */
 	EVisibility GetResourceTypeVerticesVisibility() const;
diff --git a/Engine/Source/Editor/MeshPaint/Public/MeshPaintEdMode.h b/Engine/Source/Editor/MeshPaint/Public/MeshPaintEdMode.h
--- a/Engine/Source/Editor/MeshPaint/Public/MeshPaintEdMode.h
+++ b/Engine/Source/Editor/MeshPaint/Public/MeshPaintEdMode.h
@@ -726,6 +726,8 @@ private:
 	/** Caches the currently selected actors info into CurrentlySelectedActorsMaterialInfo */
 	void CacheActorInfo();
 
+	/** Returns valid StaticMesheComponents in the current selection */
+	TArray<UStaticMeshComponent*> GetValidStaticMeshComponents() const;
 private:
 
 	/** Whether we're currently painting */
@@ -853,11 +855,12 @@ public:
 	/**
 	 * Imports Vertex Color data from texture scanning thought uv vertex coordinates for selected actors.  
 	 *
+	 * @param ModeTools the mode tools for the current editor
 	 * @param Filename path for loading TGA file.
 	 * @param UVIndex Coordinate index.
 	 * @param ImportLOD LOD level to work with.
 	 * @param Tex Texture info.
 	 * @param ColorMask Mask for filtering which colors to use.
 	 */
-	void ImportVertexColors(const FString& Filename, int32 UVIndex, int32 ImportLOD, uint8 ColorMask);
+	void ImportVertexColors(FEditorModeTools* ModeTools, const FString& Filename, int32 UVIndex, int32 ImportLOD, uint8 ColorMask);
 };
