Editor: Slowly starting to generalize mesh paint to work with other component types (WIP)

diff --git a/Engine/Source/Editor/MeshPaint/Private/MeshPaintEdMode.cpp b/Engine/Source/Editor/MeshPaint/Private/MeshPaintEdMode.cpp
--- a/Engine/Source/Editor/MeshPaint/Private/MeshPaintEdMode.cpp
+++ b/Engine/Source/Editor/MeshPaint/Private/MeshPaintEdMode.cpp
@@ -2316,7 +2316,7 @@ void FEdModeMeshPaint::PaintTexture( const FMeshPaintParameters& InParams,
 	else
 	{
 
-		// Constants used for generating quads accross entire paint rendertarget
+		// Constants used for generating quads across entire paint rendertarget
 		const float MinU = 0.0f;
 		const float MinV = 0.0f;
 		const float MaxU = 1.0f;
@@ -2491,10 +2491,10 @@ void FEdModeMeshPaint::RestoreRenderTargets()
 	bDoRestoreRenTargets = true;
 }
 
-/** Clears all texture overrides for this static mesh */
-void FEdModeMeshPaint::ClearStaticMeshTextureOverrides(UStaticMeshComponent* InStaticMeshComponent)
+/** Clears all texture overrides for this component */
+void FEdModeMeshPaint::ClearStaticMeshTextureOverrides(UMeshComponent* InMeshComponent)
 {
-	if(!InStaticMeshComponent)
+	if (!InMeshComponent)
 	{
 		return;
 	}
@@ -2504,7 +2504,7 @@ void FEdModeMeshPaint::ClearStaticMeshTextureOverrides(UStaticMeshComponent* InS
 	TArray<UMaterialInterface*> UsedMaterials;
 
 	// Get all the used materials for this StaticMeshComponent
-	InStaticMeshComponent->GetUsedMaterials( UsedMaterials );
+	InMeshComponent->GetUsedMaterials(UsedMaterials);
 
 	for( int32 MatIndex = 0; MatIndex < UsedMaterials.Num(); MatIndex++)
 	{
@@ -2542,9 +2542,9 @@ void FEdModeMeshPaint::ClearAllTextureOverrides()
 }
 
 /** Sets all texture overrides available for the mesh. */
-void FEdModeMeshPaint::SetAllTextureOverrides(UStaticMeshComponent* InStaticMeshComponent)
+void FEdModeMeshPaint::SetAllTextureOverrides(UMeshComponent* InMeshComponent)
 {
-	if(!InStaticMeshComponent)
+	if (!InMeshComponent)
 	{
 		return;
 	}
@@ -2554,7 +2554,7 @@ void FEdModeMeshPaint::SetAllTextureOverrides(UStaticMeshComponent* InStaticMesh
 	TArray<UMaterialInterface*> UsedMaterials;
 
 	// Get all the used materials for this StaticMeshComponent
-	InStaticMeshComponent->GetUsedMaterials( UsedMaterials );
+	InMeshComponent->GetUsedMaterials(UsedMaterials);
 
 	// Add the materials this actor uses to the list we maintain for ALL the selected actors, but only if
 	//  it does not appear in the list already.
@@ -2583,7 +2583,7 @@ void FEdModeMeshPaint::SetAllTextureOverrides(UStaticMeshComponent* InStaticMesh
 }
 
 /** Sets the override for a specific texture for any materials using it in the mesh, clears the override if it has no overrides. */
-void FEdModeMeshPaint::SetSpecificTextureOverrideForMesh(UStaticMeshComponent* InStaticMeshComponent, UTexture* Texture)
+void FEdModeMeshPaint::SetSpecificTextureOverrideForMesh(UMeshComponent* InMeshComponent, UTexture* Texture)
 {
 	PaintTexture2DData* TextureData = GetPaintTargetData( (UTexture2D*)Texture );
 
@@ -2591,7 +2591,7 @@ void FEdModeMeshPaint::SetSpecificTextureOverrideForMesh(UStaticMeshComponent* I
 
 	// Check all the materials on the mesh to see if the user texture is there
 	int32 MaterialIndex = 0;
-	UMaterialInterface* MaterialToCheck = InStaticMeshComponent->GetMaterial( MaterialIndex );
+	UMaterialInterface* MaterialToCheck = InMeshComponent->GetMaterial(MaterialIndex);
 	while( MaterialToCheck != NULL )
 	{
 		bool bIsTextureUsed = DoesMaterialUseTexture(MaterialToCheck,Texture);
@@ -2613,7 +2613,7 @@ void FEdModeMeshPaint::SetSpecificTextureOverrideForMesh(UStaticMeshComponent* I
 		}
 
 		++MaterialIndex;
-		MaterialToCheck = InStaticMeshComponent->GetMaterial( MaterialIndex );
+		MaterialToCheck = InMeshComponent->GetMaterial(MaterialIndex);
 	}
 }
 
@@ -4500,16 +4500,16 @@ void FEdModeMeshPaint::UpdateTexturePaintTargetList()
 				// Get the selected material index and selected actor from the cached actor info
 				const int32 MaterialIndex = MeshData->SelectedMaterialIndex;
 
-				// we only operate on static meshes.
-				TArray<UStaticMeshComponent*> StaticMeshComponents;
-				Actor->GetComponents<UStaticMeshComponent>(StaticMeshComponents);
+				// we only operate on mesh components.
+				TArray<UMeshComponent*> MeshComponents;
+				Actor->GetComponents<UMeshComponent>(MeshComponents);
 
-				for(const auto& StaticMeshComponent : StaticMeshComponents)
+				for(const auto& MeshComponent : MeshComponents)
 				{
-					if ( StaticMeshComponent != NULL )
+					if ( MeshComponent != NULL )
 					{
 						// We already know the material we are painting on, take it off the static mesh component
-						UMaterialInterface* Material = StaticMeshComponent->GetMaterial(MaterialIndex);
+						UMaterialInterface* Material = MeshComponent->GetMaterial(MaterialIndex);
 				
 						if ( Material != NULL )
 						{
@@ -5138,21 +5138,19 @@ int32 FEdModeMeshPaint::GetEditingActorsNumberOfMaterials() const
 void FEdModeMeshPaint::CacheActorInfo()
 {
 	TMap<TWeakObjectPtr<AActor>,FMeshSelectedMaterialInfo> TempMap;
-	TArray<UStaticMeshComponent*> SMComponents = GetValidStaticMeshComponents();
-	for( int32 CurSMIndex = 0; CurSMIndex < SMComponents.Num(); ++CurSMIndex )
-	{
-		TArray<UMaterialInterface*> UsedMaterials;
 
-		// Currently we only support static mesh components
-		UStaticMeshComponent* StaticMesh = SMComponents[CurSMIndex];
-
-		// Get the materials used by the mesh
-		StaticMesh->GetUsedMaterials( UsedMaterials );
-		AActor* CurActor = CastChecked< AActor >(StaticMesh->GetOuter());
+	TArray<UStaticMeshComponent*> MeshComponents = GetValidStaticMeshComponents();
+	for (const auto& MeshComponent : MeshComponents)
+	{
+		AActor* CurActor = CastChecked< AActor >(MeshComponent->GetOuter());
 		if (CurActor != NULL)
 		{
 			if (!CurrentlySelectedActorsMaterialInfo.Contains(CurActor))
 			{
+				// Get the materials used by the mesh
+				TArray<UMaterialInterface*> UsedMaterials;
+				MeshComponent->GetUsedMaterials(UsedMaterials);
+
 				TempMap.Add(CurActor,FMeshSelectedMaterialInfo(UsedMaterials.Num()));
 			}
 			else
diff --git a/Engine/Source/Editor/MeshPaint/Public/MeshPaintEdMode.h b/Engine/Source/Editor/MeshPaint/Public/MeshPaintEdMode.h
--- a/Engine/Source/Editor/MeshPaint/Public/MeshPaintEdMode.h
+++ b/Engine/Source/Editor/MeshPaint/Public/MeshPaintEdMode.h
@@ -485,15 +485,15 @@ public:
 	 */	
 	void CommitAllPaintedTextures();
 
-	/** Clears all texture overrides for this static mesh. */
-	void ClearStaticMeshTextureOverrides(UStaticMeshComponent* InStaticMeshComponent);
+	/** Clears all texture overrides for this component. */
+	void ClearStaticMeshTextureOverrides(UMeshComponent* InMeshComponent);
 
 	/** Clears all texture overrides, removing any pending texture paint changes */
 	void ClearAllTextureOverrides();
 
-	void SetAllTextureOverrides(UStaticMeshComponent* InStaticMeshComponent);
+	void SetAllTextureOverrides(UMeshComponent* InMeshComponent);
 
-	void SetSpecificTextureOverrideForMesh(UStaticMeshComponent* InStaticMeshComponent, UTexture* Texture);
+	void SetSpecificTextureOverrideForMesh(UMeshComponent* InMeshComponent, UTexture* Texture);
 
 	/** Used to tell the texture paint system that we will need to restore the rendertargets */
 	void RestoreRenderTargets();
