Editor: Generalizing mesh paint to work with other component types (WIP)

diff --git a/Engine/Source/Editor/MeshPaint/Private/MeshPaintEdMode.cpp b/Engine/Source/Editor/MeshPaint/Private/MeshPaintEdMode.cpp
--- a/Engine/Source/Editor/MeshPaint/Private/MeshPaintEdMode.cpp
+++ b/Engine/Source/Editor/MeshPaint/Private/MeshPaintEdMode.cpp
@@ -2492,7 +2492,7 @@ void FEdModeMeshPaint::RestoreRenderTargets()
 }
 
 /** Clears all texture overrides for this component */
-void FEdModeMeshPaint::ClearStaticMeshTextureOverrides(UMeshComponent* InMeshComponent)
+void FEdModeMeshPaint::ClearMeshTextureOverrides(UMeshComponent* InMeshComponent)
 {
 	if (!InMeshComponent)
 	{
@@ -2785,32 +2785,29 @@ void FEdModeMeshPaint::Render( const FSceneView* View, FViewport* Viewport, FPri
 	}
 }
 
-// @TODO MeshPaint: Cache selected static mesh components each time selection change
-/** Returns valid StaticMesheComponents in the current selection */
-TArray<UStaticMeshComponent*> FEdModeMeshPaint::GetValidStaticMeshComponents() const
+// @TODO MeshPaint: Cache selected mesh components each time selection change
+/** Returns valid MeshComponents in the current selection */
+TArray<UMeshComponent*> FEdModeMeshPaint::GetSelectedMeshComponents() const
 {
-	TArray<UStaticMeshComponent*> SMComponents;
+	TArray<UMeshComponent*> Result;
 
-	// Iterate over selected actors looking for static meshes
-	USelection& SelectedActors = *Owner->GetSelectedActors();
-	for( int32 CurSelectedActorIndex = 0; CurSelectedActorIndex < SelectedActors.Num(); ++CurSelectedActorIndex )
+	for (FSelectionIterator SelectionIt(*Owner->GetSelectedActors()); SelectionIt; ++SelectionIt)
 	{
-		AActor* CurActor = CastChecked< AActor >( SelectedActors.GetSelectedObject( CurSelectedActorIndex ) );
+		AActor* CurActor = CastChecked<AActor>(*SelectionIt);
 
 		// Ignore actors that are hidden or not selected
-		if ( CurActor->bHidden || !CurActor->IsSelected() )
+		if (CurActor->bHidden || !CurActor->IsSelected())
 		{
 			continue;
 		}
 
-		TArray<UStaticMeshComponent*> ActorMeshComponents;
-		CurActor->GetComponents<UStaticMeshComponent>( ActorMeshComponents );
-
-		SMComponents.Append( ActorMeshComponents );
+		TArray<UMeshComponent*> MeshComponents;
+		CurActor->GetComponents<UMeshComponent>(MeshComponents);
 
+		Result.Append(MeshComponents);
 	}
 
-	return SMComponents;
+	return Result;
 }
 
 /** Saves out cached mesh settings for the given actor */
@@ -2818,22 +2815,17 @@ void FEdModeMeshPaint::SaveSettingsForActor( AActor* InActor )
 {
 	if( InActor != NULL )
 	{
-		AStaticMeshActor* StaticMeshActor = Cast< AStaticMeshActor >( InActor );
+		TArray<UMeshComponent*> MeshComponents;
+		InActor->GetComponents<UMeshComponent>(MeshComponents);
 
-		UStaticMeshComponent* StaticMeshComponent = NULL;
-		if( StaticMeshActor != NULL )
-		{
-			StaticMeshComponent = StaticMeshActor->GetStaticMeshComponent();
-		}
-
-		if( StaticMeshComponent != NULL )
+		for( UMeshComponent* MeshComponent : MeshComponents )
 		{
 			// Get the currently selected texture
 			UTexture2D* SelectedTexture = GetSelectedTexture();
 
-			// Get all the used materials for this StaticMeshComponent
+			// Get all the used materials for this component
 			TArray<UMaterialInterface*> UsedMaterials;
-			StaticMeshComponent->GetUsedMaterials( UsedMaterials );
+			MeshComponent->GetUsedMaterials( UsedMaterials );
 
 			// Check this mesh's textures against the selected one before we save the settings to make sure it's a valid texture
 			for( int32 MatIndex = 0; MatIndex < UsedMaterials.Num(); MatIndex++)
@@ -2862,28 +2854,28 @@ void FEdModeMeshPaint::SaveSettingsForActor( AActor* InActor )
 					{
 						// Save the settings for this mesh with its valid texture
 						StaticMeshSettings MeshSettings = StaticMeshSettings(SelectedTexture, FMeshPaintSettings::Get().UVChannel);
-						StaticMeshSettingsMap.Add(StaticMeshComponent, MeshSettings);
+						StaticMeshSettingsMap.Add(MeshComponent, MeshSettings);
 						return;
 					}
 				}		
 			}
 
 			// No valid Texture found, attempt to find the previous texture setting or leave it as NULL to be handled by the default texture on selection
-			StaticMeshSettings* FoundMeshSettings = StaticMeshSettingsMap.Find(StaticMeshComponent);
+			StaticMeshSettings* FoundMeshSettings = StaticMeshSettingsMap.Find(MeshComponent);
 			UTexture2D* SavedTexture = FoundMeshSettings != NULL ? FoundMeshSettings->SelectedTexture : NULL;
 			StaticMeshSettings MeshSettings = StaticMeshSettings(SavedTexture, FMeshPaintSettings::Get().UVChannel);
-			StaticMeshSettingsMap.Add(StaticMeshComponent, MeshSettings);
+			StaticMeshSettingsMap.Add(MeshComponent, MeshSettings);
 		}
 	}
 }
 
-void FEdModeMeshPaint::UpdateSettingsForStaticMeshComponent( UStaticMeshComponent* InStaticMeshComponent, UTexture2D* InOldTexture, UTexture2D* InNewTexture )
+void FEdModeMeshPaint::UpdateSettingsForMeshComponent( UMeshComponent* InMeshComponent, UTexture2D* InOldTexture, UTexture2D* InNewTexture )
 {
-	if( InStaticMeshComponent != NULL )
+	if( InMeshComponent != NULL )
 	{
-		// Get all the used materials for this InStaticMeshComponent
+		// Get all the used materials for this component
 		TArray<UMaterialInterface*> UsedMaterials;
-		InStaticMeshComponent->GetUsedMaterials( UsedMaterials );
+		InMeshComponent->GetUsedMaterials(UsedMaterials);
 
 		// Check this mesh's textures against the selected one before we save the settings to make sure it's a valid texture
 		for( int32 MatIndex = 0; MatIndex < UsedMaterials.Num(); MatIndex++)
@@ -2912,7 +2904,7 @@ void FEdModeMeshPaint::UpdateSettingsForStaticMeshComponent( UStaticMeshComponen
 				{
 					// Save the settings for this mesh with its valid texture
 					StaticMeshSettings MeshSettings = StaticMeshSettings(InNewTexture, FMeshPaintSettings::Get().UVChannel);
-					StaticMeshSettingsMap.Add(InStaticMeshComponent, MeshSettings);
+					StaticMeshSettingsMap.Add(InMeshComponent, MeshSettings);
 					return;
 				}
 			}		
@@ -2934,7 +2926,7 @@ bool FEdModeMeshPaint::Select( AActor* InActor, bool bInSelected )
 				if(FMeshPaintSettings::Get().ResourceType == EMeshPaintResource::Texture)
 				{
 					// When un-selecting a mesh, save it's settings based on the current properties
-					ClearStaticMeshTextureOverrides(StaticMeshComponent);
+					ClearMeshTextureOverrides(StaticMeshComponent);
 					SaveSettingsForActor(InActor);
 				}
 				else if(FMeshPaintSettings::Get().ResourceType == EMeshPaintResource::VertexColors)
@@ -2986,11 +2978,10 @@ void FEdModeMeshPaint::ActorSelectionChangeNotify()
 		StaticMeshSettings* MeshSettings = NULL;
 
 		// For now, just grab the first mesh we find with some cached settings
-		TArray<UStaticMeshComponent*> SMComponents = GetValidStaticMeshComponents();
-		for( int32 CurSMIndex = 0; CurSMIndex < SMComponents.Num(); ++CurSMIndex )
+		TArray<UMeshComponent*> MeshComponents = GetSelectedMeshComponents();
+		for( UMeshComponent* MeshComponent : MeshComponents )
 		{
-			UStaticMeshComponent* StaticMesh = SMComponents[CurSMIndex];
-			MeshSettings = StaticMeshSettingsMap.Find(StaticMesh);
+			MeshSettings = StaticMeshSettingsMap.Find(MeshComponent);
 			if( MeshSettings != NULL )
 			{
 				break;
@@ -2999,7 +2990,6 @@ void FEdModeMeshPaint::ActorSelectionChangeNotify()
 
 		if( MeshSettings != NULL)
 		{
-
 			// Set UVChannel to our cached setting
 			FMeshPaintSettings::Get().UVChannel = MeshSettings->SelectedUVChannel;
 
@@ -3024,7 +3014,7 @@ void FEdModeMeshPaint::ActorSelectionChangeNotify()
 			// Update texture list below to reflect any selection changes
 			bShouldUpdateTextureList = true;
 		}
-		else if( SMComponents.Num() > 0 )
+		else if( MeshComponents.Num() > 0 )
 		{
 			// No cached settings, default UVChannel to 0 and Texture Target list to first selection
 			FMeshPaintSettings::Get().UVChannel = 0;
@@ -4104,11 +4094,11 @@ void FEdModeMeshPaint::RemoveInstanceMaterialAndTexture() const
 		AActor* SelectedActor = Cast< AActor >( SelectedActors.GetSelectedObject( CurSelectedActorIndex ) );
 		if(SelectedActor != NULL)
 		{
-			TArray<UStaticMeshComponent*> StaticMeshComponents;
-			SelectedActor->GetComponents<UStaticMeshComponent>(StaticMeshComponents);
-			for(const auto& StaticMeshComponent : StaticMeshComponents)
+			TArray<UMeshComponent*> StaticMeshComponents;
+			SelectedActor->GetComponents<UMeshComponent>(StaticMeshComponents);
+			for(const auto& MeshComponent : StaticMeshComponents)
 			{
-				if( StaticMeshComponent != NULL )
+				if( MeshComponent != NULL )
 				{
 					// @todo: this function
 				}
@@ -4228,27 +4218,38 @@ bool FEdModeMeshPaint::CanPasteVertexColors() const
 void FImportVertexTextureHelper::PickVertexColorFromTex(FColor& NewVertexColor, uint8* MipData, FVector2D & UV, UTexture2D* Tex, uint8& ColorMask)
 {	
 	check(MipData);
-	NewVertexColor = FColor(0,0,0, 0);
+	NewVertexColor = FColor(0,0,0,0);
 
-	if (UV.X >=0 && UV.X <1 && UV.Y >=0 && UV.Y <1)
+	if ((UV.X >= 0.0f) && (UV.X < 1.0f) && (UV.Y >= 0.0f) && (UV.Y < 1.0f))
 	{
-		int32 x =  Tex->GetSizeX()*UV.X;
-		int32 y =  Tex->GetSizeY()*UV.Y;
+		const int32 X = Tex->GetSizeX()*UV.X;
+		const int32 Y = Tex->GetSizeY()*UV.Y;
 
-		const int32 idx = ((y * Tex->GetSizeX()) + x) * 4;
-		uint8 B = MipData[idx];
-		uint8 G = MipData[idx+1];
-		uint8 R = MipData[idx+2];
-		uint8 A = MipData[idx+3];
+		const int32 Index = ((Y * Tex->GetSizeX()) + X) * 4;
+		uint8 B = MipData[Index + 0];
+		uint8 G = MipData[Index + 1];
+		uint8 R = MipData[Index + 2];
+		uint8 A = MipData[Index + 3];
 
 		if (ColorMask & ChannelsMask::ERed)
+		{
 			NewVertexColor.R = R;
+		}
+
 		if (ColorMask & ChannelsMask::EGreen)
+		{
 			NewVertexColor.G = G;
+		}
+
 		if (ColorMask & ChannelsMask::EBlue)
+		{
 			NewVertexColor.B = B;
+		}
+		
 		if (ColorMask & ChannelsMask::EAlpha)
-			NewVertexColor.A = A ;
+		{
+			NewVertexColor.A = A;
+		}
 	}	
 }
 
@@ -4669,14 +4670,22 @@ int32 FEdModeMeshPaint::GetMaxNumUVSets() const
 {
 	int32 MaxNumUVSets = 0;
 
-	// Iterate over selected static mesh components
-	TArray<UStaticMeshComponent*> SMComponents = GetValidStaticMeshComponents();
-	for( int32 CurSMIndex = 0; CurSMIndex < SMComponents.Num(); ++CurSMIndex )
+	TArray<UMeshComponent*> MeshComponents = GetSelectedMeshComponents();
+	for( UMeshComponent* MeshComponent : MeshComponents )
 	{
-		UStaticMeshComponent* StaticMeshComponent = SMComponents[ CurSMIndex ];
+		// Get the number of UV sets for this mesh
+		int32 NumUVSets = 0;
+
+		if (UStaticMeshComponent* StaticMeshComponent = Cast<UStaticMeshComponent>(MeshComponent))
+		{
+			NumUVSets = StaticMeshComponent->StaticMesh->RenderData->LODResources[PaintingMeshLODIndex].VertexBuffer.GetNumTexCoords();
+		}
+		else
+		{
+			//@TODO: MESHPAINT: Need a generic way to query this, right now assuming 1 UV set for anything other than static meshes
+			NumUVSets = 1;
+		}
 
-		// Get the number of UV sets for this static mesh
-		int32 NumUVSets = StaticMeshComponent->StaticMesh->RenderData->LODResources[PaintingMeshLODIndex].VertexBuffer.GetNumTexCoords();
 		MaxNumUVSets = FMath::Max(NumUVSets, MaxNumUVSets);
 	}
 
@@ -4920,11 +4929,11 @@ void FEdModeMeshPaint::Tick(FEditorViewportClient* ViewportClient,float DeltaTim
 	//Will set the texture override up for the selected texture, important for the drop down combo-list and selecting between material instances.
 	if(FMeshPaintSettings::Get().ResourceType == EMeshPaintResource::Texture)
 	{
-		TArray<UStaticMeshComponent*> SMComponents = GetValidStaticMeshComponents();
+		TArray<UMeshComponent*> MeshComponents = GetSelectedMeshComponents();
 
-		for( int32 CurSMIndex = 0; CurSMIndex < SMComponents.Num(); ++CurSMIndex )
+		for( UMeshComponent* MeshComponent : MeshComponents )
 		{
-			SetSpecificTextureOverrideForMesh(SMComponents[CurSMIndex ], GetSelectedTexture());
+			SetSpecificTextureOverrideForMesh(MeshComponent, GetSelectedTexture());
 		}
 	}
 
@@ -4962,7 +4971,7 @@ void FEdModeMeshPaint::DuplicateTextureMaterialCombo()
 {
 	UTexture2D* SelectedTexture = GetSelectedTexture();
 	
-	if ( NULL != SelectedTexture && ActorBeingEdited.IsValid() )
+	if ( (NULL != SelectedTexture) && ActorBeingEdited.IsValid() )
 	{
 		const FMeshSelectedMaterialInfo* MeshData = CurrentlySelectedActorsMaterialInfo.Find(ActorBeingEdited);
 		if (MeshData != NULL)
@@ -5064,10 +5073,10 @@ void FEdModeMeshPaint::DuplicateTextureMaterialCombo()
 
 				bool bMaterialChanged = false;
 				UStaticMeshComponent* SMComponent = StaticMeshActor->GetStaticMeshComponent();
-				ClearStaticMeshTextureOverrides(SMComponent);
+				ClearMeshTextureOverrides(SMComponent);
 
 				SMComponent->SetMaterial(MaterialIndex,NewMaterialInstance);
-				UpdateSettingsForStaticMeshComponent(SMComponent, SelectedTexture, NewTexture);
+				UpdateSettingsForMeshComponent(SMComponent, SelectedTexture, NewTexture);
 
 				SMComponent->MarkPackageDirty();
 			
@@ -5139,7 +5148,7 @@ void FEdModeMeshPaint::CacheActorInfo()
 {
 	TMap<TWeakObjectPtr<AActor>,FMeshSelectedMaterialInfo> TempMap;
 
-	TArray<UStaticMeshComponent*> MeshComponents = GetValidStaticMeshComponents();
+	TArray<UMeshComponent*> MeshComponents = GetSelectedMeshComponents();
 	for (const auto& MeshComponent : MeshComponents)
 	{
 		AActor* CurActor = CastChecked< AActor >(MeshComponent->GetOuter());
diff --git a/Engine/Source/Editor/MeshPaint/Public/MeshPaintEdMode.h b/Engine/Source/Editor/MeshPaint/Public/MeshPaintEdMode.h
--- a/Engine/Source/Editor/MeshPaint/Public/MeshPaintEdMode.h
+++ b/Engine/Source/Editor/MeshPaint/Public/MeshPaintEdMode.h
@@ -385,8 +385,8 @@ public:
 	/** Saves out cached mesh settings for the given actor */
 	void SaveSettingsForActor( AActor* InActor );
 
-	/** Updates static Mesh with settings for updated textures*/
-	void UpdateSettingsForStaticMeshComponent( UStaticMeshComponent* InStaticMeshComponent, UTexture2D* InOldTexture, UTexture2D* InNewTexture );
+	/** Updates settings for the specified component with updated textures */
+	void UpdateSettingsForMeshComponent( UMeshComponent* InMeshComponent, UTexture2D* InOldTexture, UTexture2D* InNewTexture );
 
 	/** Helper function to get the current paint action for use in DoPaint */
 	EMeshPaintAction::Type GetPaintAction(FViewport* InViewport);
@@ -486,7 +486,7 @@ public:
 	void CommitAllPaintedTextures();
 
 	/** Clears all texture overrides for this component. */
-	void ClearStaticMeshTextureOverrides(UMeshComponent* InMeshComponent);
+	void ClearMeshTextureOverrides(UMeshComponent* InMeshComponent);
 
 	/** Clears all texture overrides, removing any pending texture paint changes */
 	void ClearAllTextureOverrides();
@@ -726,8 +726,8 @@ private:
 	/** Caches the currently selected actors info into CurrentlySelectedActorsMaterialInfo */
 	void CacheActorInfo();
 
-	/** Returns valid StaticMesheComponents in the current selection */
-	TArray<UStaticMeshComponent*> GetValidStaticMeshComponents() const;
+	/** Returns valid MeshComponents in the current selection */
+	TArray<UMeshComponent*> GetSelectedMeshComponents() const;
 private:
 
 	/** Whether we're currently painting */
@@ -785,13 +785,13 @@ private:
 	 * @param	StaticMeshComponent		The SMC to remove vertex colors from.
 	 * @param	InstanceMeshLODInfo		The instance's LODInfo which stores the painted information to be cleared.
 	 */
-	void RemoveInstanceVertexColorsWorker(UStaticMeshComponent *StaticMeshComponent, FStaticMeshComponentLODInfo *InstanceMeshLODInfo) const;
+	void RemoveInstanceVertexColorsWorker(UStaticMeshComponent* StaticMeshComponent, FStaticMeshComponentLODInfo *InstanceMeshLODInfo) const;
 
 	/** Texture paint: Will hold a list of texture items that we can paint on */
 	TArray<FTextureTargetListInfo> TexturePaintTargetList;
 
 	/** Map of settings for each StaticMeshComponent */
-	TMap< UStaticMeshComponent*, StaticMeshSettings > StaticMeshSettingsMap;
+	TMap< UMeshComponent*, StaticMeshSettings > StaticMeshSettingsMap;
 
 	/** Used to store a flag that will tell the tick function to restore data to our rendertargets after they have been invalidated by a viewport resize. */
 	bool bDoRestoreRenTargets;
