Editor: Generalizing mesh paint to work with other component types: start of an adaptor interface for different components (WIP)

diff --git a/Engine/Source/Editor/MeshPaint/Private/MeshPaintEdMode.cpp b/Engine/Source/Editor/MeshPaint/Private/MeshPaintEdMode.cpp
--- a/Engine/Source/Editor/MeshPaint/Private/MeshPaintEdMode.cpp
+++ b/Engine/Source/Editor/MeshPaint/Private/MeshPaintEdMode.cpp
@@ -43,6 +43,74 @@ DEFINE_LOG_CATEGORY_STATIC(LogMeshPaintEdMode, Log, All);
 FMeshPaintSettings FMeshPaintSettings::StaticMeshPaintSettings;
 
 
+
+class FMeshPaintGeometryAdaptorForStaticMeshes : public IMeshPaintGeometryAdaptor
+{
+public:
+	virtual bool Construct(UMeshComponent* InComponent, int32 InPaintingMeshLODIndex, int32 InUVChannelIndex) override;
+	virtual int32 GetNumTexCoords() const override;
+	virtual void GetTriangleInfo(int32 TriIndex, FEdModeMeshPaint::FTexturePaintTriangleInfo& OutTriInfo) const override;
+
+protected:
+	UStaticMeshComponent* StaticMeshComponent;
+	FStaticMeshLODResources* LODModel;
+	int32 UVChannelIndex;
+	FIndexArrayView Indices;
+};
+
+
+
+bool FMeshPaintGeometryAdaptorForStaticMeshes::Construct(UMeshComponent* InComponent, int32 InPaintingMeshLODIndex, int32 InUVChannelIndex)
+{
+	StaticMeshComponent = Cast<UStaticMeshComponent>(InComponent);
+	if (StaticMeshComponent != nullptr)
+	{
+		if (StaticMeshComponent->StaticMesh != nullptr)
+		{
+			if (InPaintingMeshLODIndex < StaticMeshComponent->StaticMesh->GetNumLODs())
+			{
+				LODModel = &(StaticMeshComponent->StaticMesh->RenderData->LODResources[InPaintingMeshLODIndex]);
+				UVChannelIndex = InUVChannelIndex;
+				Indices = LODModel->IndexBuffer.GetArrayView();
+
+				return (UVChannelIndex >= 0) && (UVChannelIndex < (int32)LODModel->VertexBuffer.GetNumTexCoords());
+			}
+		}
+	}
+
+	return false;
+}
+
+int32 FMeshPaintGeometryAdaptorForStaticMeshes::GetNumTexCoords() const
+{
+	return LODModel->VertexBuffer.GetNumTexCoords();
+}
+
+void FMeshPaintGeometryAdaptorForStaticMeshes::GetTriangleInfo(int32 TriIndex, FEdModeMeshPaint::FTexturePaintTriangleInfo& OutTriInfo) const
+{
+	// Grab the vertex indices and points for this triangle
+	for (int32 TriVertexNum = 0; TriVertexNum < 3; ++TriVertexNum)
+	{
+		const int32 VertexIndex = Indices[TriIndex * 3 + TriVertexNum];
+		OutTriInfo.TriVertices[TriVertexNum] = LODModel->PositionVertexBuffer.VertexPosition(VertexIndex);
+		OutTriInfo.TriUVs[TriVertexNum] = LODModel->VertexBuffer.GetVertexUV(VertexIndex, UVChannelIndex);
+	}
+}
+
+TSharedPtr<IMeshPaintGeometryAdaptor> PRIVATE_ConstructMeshPaintGeometryAdaptor(UMeshComponent* InComponent, int32 InPaintingMeshLODIndex, int32 InUVChannelIndex)
+{
+	// Create the geometry adapter
+	TSharedPtr<FMeshPaintGeometryAdaptorForStaticMeshes> Adaptor = MakeShareable(new FMeshPaintGeometryAdaptorForStaticMeshes());
+
+	if (!Adaptor->Construct(InComponent, InPaintingMeshLODIndex, InUVChannelIndex))
+	{
+		Adaptor.Reset();
+	}
+
+	return Adaptor;
+}
+
+
 /** Batched element parameters for texture paint shaders used for paint blending and paint mask generation */
 class FMeshPaintBatchedElementParameters : public FBatchedElementParameters
 {
@@ -88,7 +156,7 @@ FEdModeMeshPaint::FEdModeMeshPaint()
 	  bPushInstanceColorsToMesh( false ),
 	  PaintingStartTime( 0.0 ),
 	  ModifiedStaticMeshes(),
-	  TexturePaintingStaticMeshComponent( NULL ),
+	  TexturePaintingCurrentMeshComponent( NULL ),
 	  TexturePaintingStaticMeshOctree(NULL),
 	  TexturePaintingStaticMeshLOD(0),
 	  PaintingTexture2D( NULL ),
@@ -116,7 +184,7 @@ void FEdModeMeshPaint::AddReferencedObjects( FReferenceCollector& Collector )
 	{
 		Collector.AddReferencedObject( ModifiedStaticMeshes[ Index ] );
 	}
-	Collector.AddReferencedObject( TexturePaintingStaticMeshComponent );
+	Collector.AddReferencedObject( TexturePaintingCurrentMeshComponent );
 	Collector.AddReferencedObject( PaintingTexture2D );
 	Collector.AddReferencedObject( BrushRenderTargetTexture );
 	Collector.AddReferencedObject( BrushMaskRenderTargetTexture );
@@ -950,12 +1018,13 @@ void FEdModeMeshPaint::DoPaint( const FVector& InCameraOrigin,
 			}
 
 			bool bHasStaticMesh = false;
-			TArray<UStaticMeshComponent*> StaticMeshComponents;
-			CurActor->GetComponents<UStaticMeshComponent>(StaticMeshComponents);
-			for(const auto& CurStaticMeshComponent : StaticMeshComponents)
+			TArray<UMeshComponent*> MeshComponents;
+			CurActor->GetComponents<UMeshComponent>(MeshComponents);
+			for(const auto& CurMeshComponent : MeshComponents)
 			{
-				if(CurStaticMeshComponent != NULL)
+				if(UStaticMeshComponent* CurStaticMeshComponent = Cast<UStaticMeshComponent>(CurMeshComponent))
 				{
+					//@TODO: MESHPAINT: Still requires a static mesh for a lot of stuff
 					UStaticMesh* CurStaticMesh = CurStaticMeshComponent->StaticMesh;
 					if(CurStaticMesh != NULL)
 					{
@@ -981,84 +1050,86 @@ void FEdModeMeshPaint::DoPaint( const FVector& InCameraOrigin,
 				ValidSelectedActors.Add( CurActor );
 			}
 
-			for(const auto& CurStaticMeshComponent : StaticMeshComponents)
+			for(const auto& CurMeshComponent : MeshComponents)
 			{
-				if(CurStaticMeshComponent == NULL)
-				{
-					continue;
-				}
+				UStaticMeshComponent* CurStaticMeshComponent = Cast<UStaticMeshComponent>(CurMeshComponent);
 
-				UStaticMesh* CurStaticMesh = CurStaticMeshComponent->StaticMesh;
-				if(CurStaticMesh == NULL)
-				{
-					continue;
-				}	
+				UStaticMesh* CurStaticMesh = (CurStaticMeshComponent != nullptr) ? CurStaticMeshComponent->StaticMesh : nullptr;
 
-				// Get a temporary body setup the has fully detailed collision for the line traces below.
-				TWeakObjectPtr<UBodySetup>* FindBodySetupPtr = StaticMeshToTempBodySetup.Find(CurStaticMesh);
-				TWeakObjectPtr<UBodySetup> CollideAllBodySetup;
-				if (FindBodySetupPtr && FindBodySetupPtr->IsValid())
-				{
-					// Existing temporary body setup for this mesh.
-					CollideAllBodySetup = *FindBodySetupPtr;
-				}
-				else
+				ECollisionEnabled::Type CachedCollisionType = ECollisionEnabled::NoCollision;
+				UBodySetup* RestoreBodySetup = nullptr;
+
+				if (CurStaticMesh != nullptr)
 				{
-					// No existing body setup in the cache map - create one from the mesh's main body setup.
-					UBodySetup* TempBodySetupRaw = DuplicateObject<UBodySetup>(CurStaticMesh->BodySetup, CurStaticMesh);
+					// Get a temporary body setup the has fully detailed collision for the line traces below.
+					TWeakObjectPtr<UBodySetup>* FindBodySetupPtr = StaticMeshToTempBodySetup.Find(CurStaticMesh);
+					TWeakObjectPtr<UBodySetup> CollideAllBodySetup;
+					if (FindBodySetupPtr && FindBodySetupPtr->IsValid())
+					{
+						// Existing temporary body setup for this mesh.
+						CollideAllBodySetup = *FindBodySetupPtr;
+					}
+					else
+					{
+						// No existing body setup in the cache map - create one from the mesh's main body setup.
+						UBodySetup* TempBodySetupRaw = DuplicateObject<UBodySetup>(CurStaticMesh->BodySetup, CurStaticMesh);
 
-					// Set collide all flag so that the body creates physics meshes using ALL elements from the mesh not just the collision mesh.
-					TempBodySetupRaw->bMeshCollideAll = true;
+						// Set collide all flag so that the body creates physics meshes using ALL elements from the mesh not just the collision mesh.
+						TempBodySetupRaw->bMeshCollideAll = true;
 
-					// This forces it to recreate the physics mesh.
-					TempBodySetupRaw->InvalidatePhysicsData();
+						// This forces it to recreate the physics mesh.
+						TempBodySetupRaw->InvalidatePhysicsData();
 
-					// Force it to use high detail tri-mesh for collisions.
-					TempBodySetupRaw->CollisionTraceFlag = CTF_UseComplexAsSimple;
-					TempBodySetupRaw->AggGeom.ConvexElems.Empty();
+						// Force it to use high detail tri-mesh for collisions.
+						TempBodySetupRaw->CollisionTraceFlag = CTF_UseComplexAsSimple;
+						TempBodySetupRaw->AggGeom.ConvexElems.Empty();
 
-					CollideAllBodySetup = TempBodySetupRaw;
+						CollideAllBodySetup = TempBodySetupRaw;
 
-					// cache the body setup (remove existing entry for this mesh if is one - it must be an invalid weak ptr).
-					StaticMeshToTempBodySetup.Remove(CurStaticMesh);
-					StaticMeshToTempBodySetup.Add(CurStaticMesh, CollideAllBodySetup);
-				}
-				
-				// Force the collision type to not be 'NoCollision' without it the line trace will always fail. 
-				const ECollisionEnabled::Type CachedCollisionType = CurStaticMeshComponent->BodyInstance.GetCollisionEnabled();
-				if (CachedCollisionType == ECollisionEnabled::NoCollision)
-				{
-					CurStaticMeshComponent->BodyInstance.SetCollisionEnabled(ECollisionEnabled::QueryOnly,false);
-				}
+						// cache the body setup (remove existing entry for this mesh if is one - it must be an invalid weak ptr).
+						StaticMeshToTempBodySetup.Remove(CurStaticMesh);
+						StaticMeshToTempBodySetup.Add(CurStaticMesh, CollideAllBodySetup);
+					}
 
-				// Swap the main and temp body setup on the mesh and recreate the physics state to update the body instance on the component.
-				UBodySetup* RestoreBodySetup = CurStaticMesh->BodySetup;
-				CurStaticMesh->BodySetup = CollideAllBodySetup.Get();
-				CurStaticMeshComponent->RecreatePhysicsState();
+					// Force the collision type to not be 'NoCollision' without it the line trace will always fail. 
+					CachedCollisionType = CurStaticMeshComponent->BodyInstance.GetCollisionEnabled();
+					if (CachedCollisionType == ECollisionEnabled::NoCollision)
+					{
+						CurStaticMeshComponent->BodyInstance.SetCollisionEnabled(ECollisionEnabled::QueryOnly, false);
+					}
+
+					// Swap the main and temp body setup on the mesh and recreate the physics state to update the body instance on the component.
+					RestoreBodySetup = CurStaticMesh->BodySetup;
+					CurStaticMesh->BodySetup = CollideAllBodySetup.Get();
+					CurStaticMeshComponent->RecreatePhysicsState();
+				}
 
 				// Ray trace
 				FHitResult TraceHitResult( 1.0f );
 				const FVector TraceExtent( 0.0f, 0.0f, 0.0f );
 
 				static FName DoPaintName(TEXT("DoPaint"));
-				if( CurStaticMeshComponent->LineTraceComponent(TraceHitResult, TraceStart, TraceEnd, FCollisionQueryParams(DoPaintName, true)) )
+				if( CurMeshComponent->LineTraceComponent(TraceHitResult, TraceStart, TraceEnd, FCollisionQueryParams(DoPaintName, true)) )
 				{
 					// Find the closest impact
-					if( BestTraceResult.GetActor() == NULL || ( TraceHitResult.Time < BestTraceResult.Time ) )
+					if( (BestTraceResult.GetActor() == nullptr) || (TraceHitResult.Time < BestTraceResult.Time) )
 					{
 						BestTraceResult = TraceHitResult;
 					}
 				}	
 				
-				// Reset the original collision type if we reset it. 
-				if (CachedCollisionType == ECollisionEnabled::NoCollision)
+				if (CurStaticMesh != nullptr)
 				{
-					CurStaticMeshComponent->BodyInstance.SetCollisionEnabled(CachedCollisionType,false);
-				}
+					// Reset the original collision type if we reset it. 
+					if (CachedCollisionType == ECollisionEnabled::NoCollision)
+					{
+						CurStaticMeshComponent->BodyInstance.SetCollisionEnabled(CachedCollisionType, false);
+					}
 
-				// Restore the main body setup on the mesh and recreate the physics state to update the body instance on the component.
-				CurStaticMesh->BodySetup = RestoreBodySetup;
-				CurStaticMeshComponent->RecreatePhysicsState();
+					// Restore the main body setup on the mesh and recreate the physics state to update the body instance on the component.
+					CurStaticMesh->BodySetup = RestoreBodySetup;
+					CurStaticMeshComponent->RecreatePhysicsState();
+				}
 			}
 		}
 
@@ -1126,31 +1197,17 @@ void FEdModeMeshPaint::DoPaint( const FVector& InCameraOrigin,
 		BeginTransaction( TransDesc );
 	}
 
-	// Iterate over the selected static meshes under the cursor and paint them!
+	// Iterate over the selected meshes under the cursor and paint them!
 	for( int32 CurActorIndex = 0; CurActorIndex < PaintableActors.Num(); ++CurActorIndex )
 	{
 		AActor* HitActor = Cast<AActor>( PaintableActors[ CurActorIndex ] );
 		check( HitActor );
 
-		TArray<UStaticMeshComponent*> StaticMeshComponents;
-		HitActor->GetComponents<UStaticMeshComponent>(StaticMeshComponents);
-		check(StaticMeshComponents.Num() > 0);
-		for(const auto& StaticMeshComponent : StaticMeshComponents)
+		TArray<UMeshComponent*> MeshComponents;
+		HitActor->GetComponents<UMeshComponent>(MeshComponents);
+		check(MeshComponents.Num() > 0);
+		for(const auto& MeshComponent : MeshComponents)
 		{
-			if(StaticMeshComponent == NULL)
-			{
-				continue;
-			}
-
-			UStaticMesh* StaticMesh = StaticMeshComponent->StaticMesh;
-			if(StaticMesh == NULL)
-			{
-				continue;
-			}
-
-			check( StaticMesh->GetNumLODs() > PaintingMeshLODIndex );
-			FStaticMeshLODResources& LODModel = StaticMeshComponent->StaticMesh->RenderData->LODResources[ PaintingMeshLODIndex ];
-		
 			// Brush properties
 			const float BrushDepth = BrushRadius;	// NOTE: Actually half of the total depth (like a radius)
 			const float BrushFalloffAmount = FMeshPaintSettings::Get().BrushFalloffAmount;
@@ -1174,7 +1231,7 @@ void FEdModeMeshPaint::DoPaint( const FVector& InCameraOrigin,
 
 
 			// Precache model -> world transform
-			const FMatrix ComponentToWorldMatrix = StaticMeshComponent->ComponentToWorld.ToMatrixWithScale();
+			const FMatrix ComponentToWorldMatrix = MeshComponent->ComponentToWorld.ToMatrixWithScale();
 
 
 			// Compute the camera position in actor space.  We need this later to check for
@@ -1258,23 +1315,30 @@ void FEdModeMeshPaint::DoPaint( const FVector& InCameraOrigin,
 					Params.PaintWeightIndex = FMath::Clamp( PaintWeightIndex, 0, Params.TotalWeightCount - 1 );
 				}
 
-				// @todo MeshPaint: Ideally we would default to: TexturePaintingStaticMeshComponent->StaticMesh->LightMapCoordinateIndex
+				// @todo MeshPaint: Ideally we would default to: TexturePaintingCurrentMeshComponent->StaticMesh->LightMapCoordinateIndex
 				//		Or we could indicate in the GUI which channel is the light map set (button to set it?)
 				Params.UVChannel = FMeshPaintSettings::Get().UVChannel;
 			}
 
-			if( FMeshPaintSettings::Get().ResourceType == EMeshPaintResource::VertexColors )
+			if (FMeshPaintSettings::Get().ResourceType == EMeshPaintResource::VertexColors)
 			{
-				// Painting vertex colors
-				PaintMeshVertices( StaticMeshComponent, Params, bShouldApplyPaint, LODModel, ComponentSpaceCameraPosition, ComponentToWorldMatrix, PDI, VisualBiasDistance);
+				if (UStaticMeshComponent* StaticMeshComponent = Cast<UStaticMeshComponent>(MeshComponent))
+				{
+					if (UStaticMesh* StaticMesh = StaticMeshComponent->StaticMesh)
+					{
+						check(StaticMesh->GetNumLODs() > PaintingMeshLODIndex);
+						FStaticMeshLODResources& LODModel = StaticMeshComponent->StaticMesh->RenderData->LODResources[PaintingMeshLODIndex];
 
+						// Painting vertex colors
+						PaintMeshVertices(StaticMeshComponent, Params, bShouldApplyPaint, LODModel, ComponentSpaceCameraPosition, ComponentToWorldMatrix, PDI, VisualBiasDistance);
+					}
+				}
 			}
 			else
 			{
 				// Painting textures
-				PaintMeshTexture( StaticMeshComponent, Params, bShouldApplyPaint, LODModel, ComponentSpaceCameraPosition, ComponentToWorldMatrix, ComponentSpaceSquaredBrushRadius, ComponentSpaceBrushPosition );
-
-			}		
+				PaintMeshTexture(MeshComponent, Params, bShouldApplyPaint, ComponentSpaceCameraPosition, ComponentToWorldMatrix, ComponentSpaceSquaredBrushRadius, ComponentSpaceBrushPosition);
+			}
 		}
 		
 
@@ -1714,7 +1778,7 @@ void FEdModeMeshPaint::PaintMeshVertices(
 
 
 /** Paints mesh texture */
-void FEdModeMeshPaint::PaintMeshTexture( UStaticMeshComponent* StaticMeshComponent, const FMeshPaintParameters& Params, const bool bShouldApplyPaint, FStaticMeshLODResources& InLODModel, const FVector& ComponentSpaceCameraPosition, const FMatrix& ComponentToWorldMatrix, const float ComponentSpaceSquaredBrushRadius, const FVector& ComponentSpaceBrushPosition )
+void FEdModeMeshPaint::PaintMeshTexture( UMeshComponent* MeshComponent, const FMeshPaintParameters& Params, const bool bShouldApplyPaint, const FVector& ComponentSpaceCameraPosition, const FMatrix& ComponentToWorldMatrix, const float ComponentSpaceSquaredBrushRadius, const FVector& ComponentSpaceBrushPosition )
 {
 	UTexture2D* TargetTexture2D = GetSelectedTexture();
 	
@@ -1727,12 +1791,28 @@ void FEdModeMeshPaint::PaintMeshTexture( UStaticMeshComponent* StaticMeshCompone
 	const bool bOnlyFrontFacing = FMeshPaintSettings::Get().bOnlyFrontFacingTriangles;
 	if( bShouldApplyPaint )
 	{
+		// Create the geometry adaptor
+		TSharedPtr<IMeshPaintGeometryAdaptor> GeometryInfo = PRIVATE_ConstructMeshPaintGeometryAdaptor(MeshComponent, PaintingMeshLODIndex, Params.UVChannel);
+		if (!GeometryInfo.IsValid())
+		{
+			return;
+		}
+
 		// @todo MeshPaint: Use a spatial database to reduce the triangle set here (kdop)
+		UStaticMeshComponent* StaticMeshComponent = Cast<UStaticMeshComponent>(MeshComponent);
+
+		if ((StaticMeshComponent == nullptr) || (StaticMeshComponent->StaticMesh == nullptr))
+		{
+			return;
+		}
 
+		check(StaticMeshComponent->StaticMesh->GetNumLODs() > PaintingMeshLODIndex);
+		FStaticMeshLODResources& LODModel = StaticMeshComponent->StaticMesh->RenderData->LODResources[PaintingMeshLODIndex];
+		const int32 NumSections = LODModel.Sections.Num();
 
 
 		// Make sure we're dealing with triangle lists
-		FIndexArrayView Indices = InLODModel.IndexBuffer.GetArrayView();
+		FIndexArrayView Indices = LODModel.IndexBuffer.GetArrayView();
 		const uint32 NumIndexBufferIndices = Indices.Num();
 		check( NumIndexBufferIndices % 3 == 0 );
 		const uint32 NumTriangles = NumIndexBufferIndices / 3;
@@ -1745,8 +1825,6 @@ void FEdModeMeshPaint::PaintMeshTexture( UStaticMeshComponent* StaticMeshCompone
 		// definitely don't want our brush to be cut off by a hard triangle edge
 		const float SquaredRadiusBias = ComponentSpaceSquaredBrushRadius * 0.025f;
 
-		FStaticMeshLODResources& LODModel = StaticMeshComponent->StaticMesh->RenderData->LODResources[PaintingMeshLODIndex];
-		int32 NumSections = LODModel.Sections.Num();
 		
 		PaintTexture2DData* TextureData = GetPaintTargetData( TargetTexture2D );
 
@@ -1773,7 +1851,7 @@ void FEdModeMeshPaint::PaintMeshTexture( UStaticMeshComponent* StaticMeshCompone
 			}
 		}
 
-		if(TexturePaintingStaticMeshOctree != NULL && (TexturePaintingStaticMeshComponent != StaticMeshComponent || TexturePaintingStaticMeshLOD != PaintingMeshLODIndex) )
+		if( (TexturePaintingStaticMeshOctree != NULL) && ((TexturePaintingCurrentMeshComponent != MeshComponent) || (TexturePaintingStaticMeshLOD != PaintingMeshLODIndex)) )
 		{
 			delete TexturePaintingStaticMeshOctree;
 			TexturePaintingStaticMeshOctree = NULL;
@@ -1870,21 +1948,21 @@ void FEdModeMeshPaint::PaintMeshTexture( UStaticMeshComponent* StaticMeshCompone
 
 		{
 
-			if( TexturePaintingStaticMeshComponent != NULL && TexturePaintingStaticMeshComponent != StaticMeshComponent )
+			if( (TexturePaintingCurrentMeshComponent != NULL) && (TexturePaintingCurrentMeshComponent != MeshComponent) )
 			{
 				// Mesh has changed, so finish up with our previous texture
 				FinishPaintingTexture();
 				bIsPainting = false;
 			}
 
-			if( TexturePaintingStaticMeshComponent == NULL )
+			if( TexturePaintingCurrentMeshComponent == NULL )
 			{
-				StartPaintingTexture( StaticMeshComponent );
+				StartPaintingTexture( MeshComponent );
 			}
 
-			if( TexturePaintingStaticMeshComponent != NULL )
+			if( TexturePaintingCurrentMeshComponent != NULL )
 			{
-				PaintTexture( Params, InfluencedTriangles, ComponentToWorldMatrix );
+				PaintTexture( Params, InfluencedTriangles, ComponentToWorldMatrix, *GeometryInfo );
 			}
 		}
 	}
@@ -1894,10 +1972,10 @@ void FEdModeMeshPaint::PaintMeshTexture( UStaticMeshComponent* StaticMeshCompone
 
 
 /** Starts painting a texture */
-void FEdModeMeshPaint::StartPaintingTexture( UStaticMeshComponent* InStaticMeshComponent )
+void FEdModeMeshPaint::StartPaintingTexture( UMeshComponent* InMeshComponent )
 {
-	check( InStaticMeshComponent != NULL );
-	check( TexturePaintingStaticMeshComponent == NULL );
+	check( InMeshComponent != NULL );
+	check( TexturePaintingCurrentMeshComponent == NULL );
 	check( PaintingTexture2D == NULL );
 
 	const auto FeatureLevel = GetWorld()->FeatureLevel;
@@ -1913,7 +1991,7 @@ void FEdModeMeshPaint::StartPaintingTexture( UStaticMeshComponent* InStaticMeshC
 
 	// Check all the materials on the mesh to see if the user texture is there
 	int32 MaterialIndex = 0;
-	UMaterialInterface* MaterialToCheck = InStaticMeshComponent->GetMaterial(MaterialIndex);
+	UMaterialInterface* MaterialToCheck = InMeshComponent->GetMaterial(MaterialIndex);
 	while( MaterialToCheck != NULL )
 	{
 		bool bIsTextureUsed = DoesMaterialUseTexture(MaterialToCheck,Texture2D);
@@ -2038,12 +2116,12 @@ void FEdModeMeshPaint::StartPaintingTexture( UStaticMeshComponent* InStaticMeshC
 		}
 
 		MaterialIndex++;
-		MaterialToCheck = InStaticMeshComponent->GetMaterial(MaterialIndex);
+		MaterialToCheck = InMeshComponent->GetMaterial(MaterialIndex);
 	}
 
 	if( bStartedPainting )
 	{
-		TexturePaintingStaticMeshComponent = InStaticMeshComponent;
+		TexturePaintingCurrentMeshComponent = InMeshComponent;
 
 		check( Texture2D != NULL );
 		PaintingTexture2D = Texture2D;
@@ -2054,11 +2132,11 @@ void FEdModeMeshPaint::StartPaintingTexture( UStaticMeshComponent* InStaticMeshC
 }
 
 
-
 /** Paints on a texture */
 void FEdModeMeshPaint::PaintTexture( const FMeshPaintParameters& InParams,
 									 const TArray< int32 >& InInfluencedTriangles,
-									 const FMatrix& InComponentToWorldMatrix )
+									 const FMatrix& InComponentToWorldMatrix,
+									 const IMeshPaintGeometryAdaptor& GeometryInfo)
 {
 	// We bail early if there are no influenced triangles
 	if( InInfluencedTriangles.Num() <= 0 )
@@ -2066,17 +2144,8 @@ void FEdModeMeshPaint::PaintTexture( const FMeshPaintParameters& InParams,
 		return;
 	}
 
-	FStaticMeshLODResources& LODModel = TexturePaintingStaticMeshComponent->StaticMesh->RenderData->LODResources[ PaintingMeshLODIndex ];
-	FIndexArrayView Indices = LODModel.IndexBuffer.GetArrayView();
-	const uint32 PaintUVCoordinateIndex = InParams.UVChannel;
 	const auto FeatureLevel = GetWorld()->FeatureLevel;
 
-	// Check to see if the UV set is available on the LOD model, if not then there is no point in continuing.
-	if( PaintUVCoordinateIndex >= LODModel.VertexBuffer.GetNumTexCoords() )
-	{
-		// @todo MeshPaint: Do we want to give the user some sort of indication that the paint failed because the UV set is not available on the object?
-		return;
-	}
 
 	PaintTexture2DData* TextureData = GetPaintTargetData( PaintingTexture2D );
 	check( TextureData != NULL && TextureData->PaintRenderTargetTexture != NULL );
@@ -2168,16 +2237,16 @@ void FEdModeMeshPaint::PaintTexture( const FMeshPaintParameters& InParams,
 	{
 		const int32 TriIndex = InInfluencedTriangles[ CurIndex ];
 		FTexturePaintTriangleInfo CurTriangle;
+		GeometryInfo.GetTriangleInfo(TriIndex, CurTriangle);
 
 		FVector2D UVMin( 99999.9f, 99999.9f );
 		FVector2D UVMax( -99999.9f, -99999.9f );
 			 
-		// Grab the vertex indices and points for this triangle
+		// Transform the triangle and update the UV bounds
 		for( int32 TriVertexNum = 0; TriVertexNum < 3; ++TriVertexNum )
-		{																		 
-			const int32 VertexIndex = Indices[ TriIndex * 3 + TriVertexNum ];
-			CurTriangle.TriVertices[ TriVertexNum ] = InComponentToWorldMatrix.TransformPosition( LODModel.PositionVertexBuffer.VertexPosition( VertexIndex ) );
-			CurTriangle.TriUVs[ TriVertexNum ] = LODModel.VertexBuffer.GetVertexUV( VertexIndex, PaintUVCoordinateIndex );
+		{				
+			// Transform the triangle to world space
+			CurTriangle.TriVertices[ TriVertexNum ] = InComponentToWorldMatrix.TransformPosition( CurTriangle.TriVertices[ TriVertexNum ] );
 
 			// Update bounds
 			float U = CurTriangle.TriUVs[ TriVertexNum ].X;
@@ -2330,7 +2399,7 @@ void FEdModeMeshPaint::PaintTexture( const FMeshPaintParameters& InParams,
 		{
 			// Generate the texture seam mask.  This is a slow operation when the object has many triangles so we only do it
 			//  once when painting is started.
-			GenerateSeamMask(TexturePaintingStaticMeshComponent, InParams.UVChannel, SeamMaskRenderTargetTexture);
+			GenerateSeamMask(TexturePaintingCurrentMeshComponent, InParams.UVChannel, SeamMaskRenderTargetTexture);
 			bGenerateSeamMask = false;
 		}
 		
@@ -2637,7 +2706,7 @@ int32 FEdModeMeshPaint::GetNumberOfPendingPaintChanges()
 /** Finishes painting a texture */
 void FEdModeMeshPaint::FinishPaintingTexture( )
 {
-	if( TexturePaintingStaticMeshComponent != NULL )
+	if( TexturePaintingCurrentMeshComponent != NULL )
 	{
 		check( PaintingTexture2D != NULL );
 
@@ -2682,7 +2751,7 @@ void FEdModeMeshPaint::FinishPaintingTexture( )
 		}
 
 		PaintingTexture2D = NULL;
-		TexturePaintingStaticMeshComponent = NULL;
+		TexturePaintingCurrentMeshComponent = NULL;
 
 		if(!bIsPainting && TexturePaintingStaticMeshOctree != NULL)
 		{
@@ -3316,8 +3385,14 @@ void FEdModeMeshPaint::CopyTextureToRenderTargetTexture( UTexture* SourceTexture
 }
 
 /** Will generate a mask texture, used for texture dilation, and store it in the passed in rendertarget */
-bool FEdModeMeshPaint::GenerateSeamMask(UStaticMeshComponent* StaticMeshComponent, int32 UVSet, UTextureRenderTarget2D* RenderTargetTexture)
+bool FEdModeMeshPaint::GenerateSeamMask(UMeshComponent* MeshComponent, int32 UVSet, UTextureRenderTarget2D* RenderTargetTexture)
 {
+	UStaticMeshComponent* StaticMeshComponent = Cast<UStaticMeshComponent>(MeshComponent);
+	if (StaticMeshComponent == nullptr)
+	{
+		return false;
+	}
+
 	check(StaticMeshComponent != NULL);
 	check(StaticMeshComponent->StaticMesh != NULL);
 	check(RenderTargetTexture != NULL);
diff --git a/Engine/Source/Editor/MeshPaint/Public/MeshPaintEdMode.h b/Engine/Source/Editor/MeshPaint/Public/MeshPaintEdMode.h
--- a/Engine/Source/Editor/MeshPaint/Public/MeshPaintEdMode.h
+++ b/Engine/Source/Editor/MeshPaint/Public/MeshPaintEdMode.h
@@ -557,7 +557,7 @@ private:
 		}
 	};
 
-	/** Triangle for use in Octree for mesh paint optimisation */
+	/** Triangle for use in Octree for mesh paint optimization */
 	struct FMeshTriangle
 	{
 		uint32 Index;
@@ -633,7 +633,7 @@ private:
 	void PaintMeshVertices( UStaticMeshComponent* StaticMeshComponent, const FMeshPaintParameters& Params, const bool bShouldApplyPaint, FStaticMeshLODResources& LODModel, const FVector& ActorSpaceCameraPosition, const FMatrix& ActorToWorldMatrix, FPrimitiveDrawInterface* PDI, const float VisualBiasDistance );
 
 	/** Paints mesh texture */
-	void PaintMeshTexture( UStaticMeshComponent* StaticMeshComponent, const FMeshPaintParameters& Params, const bool bShouldApplyPaint, FStaticMeshLODResources& LODModel, const FVector& ActorSpaceCameraPosition, const FMatrix& ActorToWorldMatrix, const float ActorSpaceSquaredBrushRadius, const FVector& ActorSpaceBrushPosition );
+	void PaintMeshTexture( UMeshComponent* MeshComponent, const FMeshPaintParameters& Params, const bool bShouldApplyPaint, const FVector& ActorSpaceCameraPosition, const FMatrix& ActorToWorldMatrix, const float ActorSpaceSquaredBrushRadius, const FVector& ActorSpaceBrushPosition );
 	
 	/** Forces real-time perspective viewports */
 	void ForceRealTimeViewports( const bool bEnable, const bool bStoreCurrentState );
@@ -642,12 +642,13 @@ private:
 	void SetViewportShowFlags( const bool bAllowColorViewModes, FEditorViewportClient& Viewport );
 
 	/** Starts painting a texture */
-	void StartPaintingTexture( UStaticMeshComponent* InStaticMeshComponent );
+	void StartPaintingTexture( UMeshComponent* InMeshComponent );
 
 	/** Paints on a texture */
 	void PaintTexture( const FMeshPaintParameters& InParams,
 					   const TArray< int32 >& InInfluencedTriangles,
-					   const FMatrix& InActorToWorldMatrix );
+					   const FMatrix& InActorToWorldMatrix,
+					   const class IMeshPaintGeometryAdaptor& GeometryInfo);
 
 	/** Finishes painting a texture */
 	void FinishPaintingTexture();
@@ -659,7 +660,7 @@ private:
 	static void CopyTextureToRenderTargetTexture(UTexture* SourceTexture, UTextureRenderTarget2D* RenderTargetTexture, ERHIFeatureLevel::Type FeatureLevel);
 
 	/** Will generate a mask texture, used for texture dilation, and store it in the passed in rendertarget */
-	bool GenerateSeamMask(UStaticMeshComponent* StaticMeshComponent, int32 UVSet, UTextureRenderTarget2D* RenderTargetTexture);
+	bool GenerateSeamMask(UMeshComponent* MeshComponent, int32 UVSet, UTextureRenderTarget2D* RenderTargetTexture);
 
 	/** Static: Creates a temporary texture used to transfer data to a render target in memory */
 	UTexture2D* CreateTempUncompressedTexture( UTexture2D* SourceTexture );
@@ -761,8 +762,8 @@ private:
 	// @todo MeshPaint: Allow painting on other LODs?
 	static const uint32 PaintingMeshLODIndex = 0;
 
-	/** Texture paint: The static mesh components that we're currently painting */
-	UStaticMeshComponent* TexturePaintingStaticMeshComponent;
+	/** Texture paint: The mesh components that we're currently painting */
+	UMeshComponent* TexturePaintingCurrentMeshComponent;
 
 	/** Texture paint: An octree for the current static mesh & LOD to speed up triangle selection */
 	FMeshTriOctree* TexturePaintingStaticMeshOctree;
@@ -864,3 +865,12 @@ public:
 	 */
 	void ImportVertexColors(FEditorModeTools* ModeTools, const FString& Filename, int32 UVIndex, int32 ImportLOD, uint8 ColorMask);
 };
+
+
+class IMeshPaintGeometryAdaptor
+{
+public:
+	virtual bool Construct(UMeshComponent* InComponent, int32 InPaintingMeshLODIndex, int32 InUVChannelIndex) = 0;
+	virtual int32 GetNumTexCoords() const = 0;
+	virtual void GetTriangleInfo(int32 TriIndex, FEdModeMeshPaint::FTexturePaintTriangleInfo& OutTriInfo) const = 0;
+};
