Editor: Add adapter factory registration system for mesh painting and split the static mesh version out into separate files (WIP)

diff --git a/Engine/Source/Editor/MeshPaint/Private/MeshPaintAdapterFactory.cpp b/Engine/Source/Editor/MeshPaint/Private/MeshPaintAdapterFactory.cpp
new file mode 100644
--- /dev/null
+++ b/Engine/Source/Editor/MeshPaint/Private/MeshPaintAdapterFactory.cpp
@@ -0,0 +1,24 @@
+// Copyright 1998-2015 Epic Games, Inc. All Rights Reserved.
+
+#include "MeshPaintPrivatePCH.h"
+#include "MeshPaintAdapterFactory.h"
+
+
+TArray<TSharedPtr<class IMeshPaintGeometryAdapterFactory>> FMeshPaintAdapterFactory::FactoryList;
+
+TSharedPtr<class IMeshPaintGeometryAdapter> FMeshPaintAdapterFactory::CreateAdapterForMesh(UMeshComponent* InComponent, int32 InPaintingMeshLODIndex, int32 InUVChannelIndex)
+{
+	TSharedPtr<IMeshPaintGeometryAdapter> Result;
+
+	for (const auto& Factory : FactoryList)
+	{
+		Result = Factory->Construct(InComponent, InPaintingMeshLODIndex, InUVChannelIndex);
+
+		if (Result.IsValid())
+		{
+			break;
+		}
+	}
+
+	return Result;
+}
diff --git a/Engine/Source/Editor/MeshPaint/Private/MeshPaintAdapterFactory.h b/Engine/Source/Editor/MeshPaint/Private/MeshPaintAdapterFactory.h
new file mode 100644
--- /dev/null
+++ b/Engine/Source/Editor/MeshPaint/Private/MeshPaintAdapterFactory.h
@@ -0,0 +1,14 @@
+// Copyright 1998-2015 Epic Games, Inc. All Rights Reserved.
+
+#pragma once
+
+#include "MeshPaintModule.h"
+
+class FMeshPaintAdapterFactory
+{
+public:
+	static TArray<TSharedPtr<class IMeshPaintGeometryAdapterFactory>> FactoryList;
+
+public:
+	static TSharedPtr<class IMeshPaintGeometryAdapter> CreateAdapterForMesh(UMeshComponent* InComponent, int32 InPaintingMeshLODIndex, int32 InUVChannelIndex);
+};
\ No newline at end of file
diff --git a/Engine/Source/Editor/MeshPaint/Private/MeshPaintEdMode.cpp b/Engine/Source/Editor/MeshPaint/Private/MeshPaintEdMode.cpp
--- a/Engine/Source/Editor/MeshPaint/Private/MeshPaintEdMode.cpp
+++ b/Engine/Source/Editor/MeshPaint/Private/MeshPaintEdMode.cpp
@@ -34,6 +34,7 @@
 #include "EngineUtils.h"
 #include "Engine/StaticMeshActor.h"
 #include "Materials/MaterialInstanceConstant.h"
+#include "MeshPaintAdapterFactory.h"
 
 #define LOCTEXT_NAMESPACE "MeshPaint_Mode"
 
@@ -42,75 +43,6 @@ DEFINE_LOG_CATEGORY_STATIC(LogMeshPaintEdMode, Log, All);
 /** Static: Global mesh paint settings */
 FMeshPaintSettings FMeshPaintSettings::StaticMeshPaintSettings;
 
-
-
-class FMeshPaintGeometryAdaptorForStaticMeshes : public IMeshPaintGeometryAdaptor
-{
-public:
-	virtual bool Construct(UMeshComponent* InComponent, int32 InPaintingMeshLODIndex, int32 InUVChannelIndex) override;
-	virtual int32 GetNumTexCoords() const override;
-	virtual void GetTriangleInfo(int32 TriIndex, FEdModeMeshPaint::FTexturePaintTriangleInfo& OutTriInfo) const override;
-
-protected:
-	UStaticMeshComponent* StaticMeshComponent;
-	FStaticMeshLODResources* LODModel;
-	int32 UVChannelIndex;
-	FIndexArrayView Indices;
-};
-
-
-
-bool FMeshPaintGeometryAdaptorForStaticMeshes::Construct(UMeshComponent* InComponent, int32 InPaintingMeshLODIndex, int32 InUVChannelIndex)
-{
-	StaticMeshComponent = Cast<UStaticMeshComponent>(InComponent);
-	if (StaticMeshComponent != nullptr)
-	{
-		if (StaticMeshComponent->StaticMesh != nullptr)
-		{
-			if (InPaintingMeshLODIndex < StaticMeshComponent->StaticMesh->GetNumLODs())
-			{
-				LODModel = &(StaticMeshComponent->StaticMesh->RenderData->LODResources[InPaintingMeshLODIndex]);
-				UVChannelIndex = InUVChannelIndex;
-				Indices = LODModel->IndexBuffer.GetArrayView();
-
-				return (UVChannelIndex >= 0) && (UVChannelIndex < (int32)LODModel->VertexBuffer.GetNumTexCoords());
-			}
-		}
-	}
-
-	return false;
-}
-
-int32 FMeshPaintGeometryAdaptorForStaticMeshes::GetNumTexCoords() const
-{
-	return LODModel->VertexBuffer.GetNumTexCoords();
-}
-
-void FMeshPaintGeometryAdaptorForStaticMeshes::GetTriangleInfo(int32 TriIndex, FEdModeMeshPaint::FTexturePaintTriangleInfo& OutTriInfo) const
-{
-	// Grab the vertex indices and points for this triangle
-	for (int32 TriVertexNum = 0; TriVertexNum < 3; ++TriVertexNum)
-	{
-		const int32 VertexIndex = Indices[TriIndex * 3 + TriVertexNum];
-		OutTriInfo.TriVertices[TriVertexNum] = LODModel->PositionVertexBuffer.VertexPosition(VertexIndex);
-		OutTriInfo.TriUVs[TriVertexNum] = LODModel->VertexBuffer.GetVertexUV(VertexIndex, UVChannelIndex);
-	}
-}
-
-TSharedPtr<IMeshPaintGeometryAdaptor> PRIVATE_ConstructMeshPaintGeometryAdaptor(UMeshComponent* InComponent, int32 InPaintingMeshLODIndex, int32 InUVChannelIndex)
-{
-	// Create the geometry adapter
-	TSharedPtr<FMeshPaintGeometryAdaptorForStaticMeshes> Adaptor = MakeShareable(new FMeshPaintGeometryAdaptorForStaticMeshes());
-
-	if (!Adaptor->Construct(InComponent, InPaintingMeshLODIndex, InUVChannelIndex))
-	{
-		Adaptor.Reset();
-	}
-
-	return Adaptor;
-}
-
-
 /** Batched element parameters for texture paint shaders used for paint blending and paint mask generation */
 class FMeshPaintBatchedElementParameters : public FBatchedElementParameters
 {
@@ -994,9 +926,12 @@ void FEdModeMeshPaint::DoPaint( const FVector& InCameraOrigin,
 {
 	const float BrushRadius = GetBrushRadiiDefault();
 
-	// Fire out a ray to see if there is a *selected* static mesh under the mouse cursor.
-	// NOTE: We can't use a GWorld line check for this as that would ignore actors that have collision disabled
-	TArray< AActor* > PaintableActors;
+	// Map of components that we could potentially interact with
+	TMap<UMeshComponent*, TSharedPtr<IMeshPaintGeometryAdapter>> PotentialComponentMap;
+	TArray<UMeshComponent*> PaintableComponents;
+
+	// Fire out a ray to see if there is a *selected* component under the mouse cursor that can be painted.
+	// NOTE: We can't use a GWorld line check for this as that would ignore components that have collision disabled
 	FHitResult BestTraceResult;
 	{
 		const FVector TraceStart( InRayOrigin );
@@ -1004,7 +939,6 @@ void FEdModeMeshPaint::DoPaint( const FVector& InCameraOrigin,
 
 		// Iterate over selected actors looking for static meshes
 		USelection& SelectedActors = *Owner->GetSelectedActors();
-		TArray< AActor* > ValidSelectedActors;
 		for( int32 CurSelectedActorIndex = 0; CurSelectedActorIndex < SelectedActors.Num(); ++CurSelectedActorIndex )
 		{
 			bool bHasKDOPTree = true;
@@ -1020,149 +954,71 @@ void FEdModeMeshPaint::DoPaint( const FVector& InCameraOrigin,
 			bool bHasStaticMesh = false;
 			TArray<UMeshComponent*> MeshComponents;
 			CurActor->GetComponents<UMeshComponent>(MeshComponents);
-			for(const auto& CurMeshComponent : MeshComponents)
+			for (const auto& CurMeshComponent : MeshComponents)
 			{
-				if(UStaticMeshComponent* CurStaticMeshComponent = Cast<UStaticMeshComponent>(CurMeshComponent))
+				// Create the geometry adapter
+				TSharedPtr<IMeshPaintGeometryAdapter> MeshAdapter = FMeshPaintAdapterFactory::CreateAdapterForMesh(CurMeshComponent, PaintingMeshLODIndex, FMeshPaintSettings::Get().UVChannel);
+				if (MeshAdapter.IsValid())
 				{
-					//@TODO: MESHPAINT: Still requires a static mesh for a lot of stuff
-					UStaticMesh* CurStaticMesh = CurStaticMeshComponent->StaticMesh;
-					if(CurStaticMesh != NULL)
-					{
-						bHasStaticMesh = true;
-						break;
-					}
-				}				
-			}
-
-			if(bHasStaticMesh)
-			{
-				if (InPaintAction == EMeshPaintAction::Fill)
-				{
-					PaintableActors.Add( CurActor );
-					continue;
-				}
-				else if (InPaintAction == EMeshPaintAction::PushInstanceColorsToMesh)
-				{
-					PaintableActors.Add( CurActor );
-					continue;
-				}
-
-				ValidSelectedActors.Add( CurActor );
-			}
-
-			for(const auto& CurMeshComponent : MeshComponents)
-			{
-				UStaticMeshComponent* CurStaticMeshComponent = Cast<UStaticMeshComponent>(CurMeshComponent);
-
-				UStaticMesh* CurStaticMesh = (CurStaticMeshComponent != nullptr) ? CurStaticMeshComponent->StaticMesh : nullptr;
-
-				ECollisionEnabled::Type CachedCollisionType = ECollisionEnabled::NoCollision;
-				UBodySetup* RestoreBodySetup = nullptr;
+					PotentialComponentMap.Add(CurMeshComponent, MeshAdapter);
 
-				if (CurStaticMesh != nullptr)
-				{
-					// Get a temporary body setup the has fully detailed collision for the line traces below.
-					TWeakObjectPtr<UBodySetup>* FindBodySetupPtr = StaticMeshToTempBodySetup.Find(CurStaticMesh);
-					TWeakObjectPtr<UBodySetup> CollideAllBodySetup;
-					if (FindBodySetupPtr && FindBodySetupPtr->IsValid())
-					{
-						// Existing temporary body setup for this mesh.
-						CollideAllBodySetup = *FindBodySetupPtr;
-					}
-					else
+					//@TODO: MESHPAINT: This is copied from the original logic, but the split between this loop and the next still feels a bit off
+					if (InPaintAction == EMeshPaintAction::Fill)
 					{
-						// No existing body setup in the cache map - create one from the mesh's main body setup.
-						UBodySetup* TempBodySetupRaw = DuplicateObject<UBodySetup>(CurStaticMesh->BodySetup, CurStaticMesh);
-
-						// Set collide all flag so that the body creates physics meshes using ALL elements from the mesh not just the collision mesh.
-						TempBodySetupRaw->bMeshCollideAll = true;
-
-						// This forces it to recreate the physics mesh.
-						TempBodySetupRaw->InvalidatePhysicsData();
-
-						// Force it to use high detail tri-mesh for collisions.
-						TempBodySetupRaw->CollisionTraceFlag = CTF_UseComplexAsSimple;
-						TempBodySetupRaw->AggGeom.ConvexElems.Empty();
-
-						CollideAllBodySetup = TempBodySetupRaw;
-
-						// cache the body setup (remove existing entry for this mesh if is one - it must be an invalid weak ptr).
-						StaticMeshToTempBodySetup.Remove(CurStaticMesh);
-						StaticMeshToTempBodySetup.Add(CurStaticMesh, CollideAllBodySetup);
+						PaintableComponents.AddUnique(CurMeshComponent);
 					}
-
-					// Force the collision type to not be 'NoCollision' without it the line trace will always fail. 
-					CachedCollisionType = CurStaticMeshComponent->BodyInstance.GetCollisionEnabled();
-					if (CachedCollisionType == ECollisionEnabled::NoCollision)
+					else if (InPaintAction == EMeshPaintAction::PushInstanceColorsToMesh)
 					{
-						CurStaticMeshComponent->BodyInstance.SetCollisionEnabled(ECollisionEnabled::QueryOnly, false);
+						PaintableComponents.AddUnique(CurMeshComponent);
 					}
 
-					// Swap the main and temp body setup on the mesh and recreate the physics state to update the body instance on the component.
-					RestoreBodySetup = CurStaticMesh->BodySetup;
-					CurStaticMesh->BodySetup = CollideAllBodySetup.Get();
-					CurStaticMeshComponent->RecreatePhysicsState();
-				}
+					// Ray trace
+					FHitResult TraceHitResult(1.0f);
 
-				// Ray trace
-				FHitResult TraceHitResult( 1.0f );
-				const FVector TraceExtent( 0.0f, 0.0f, 0.0f );
-
-				static FName DoPaintName(TEXT("DoPaint"));
-				if( CurMeshComponent->LineTraceComponent(TraceHitResult, TraceStart, TraceEnd, FCollisionQueryParams(DoPaintName, true)) )
-				{
-					// Find the closest impact
-					if( (BestTraceResult.GetActor() == nullptr) || (TraceHitResult.Time < BestTraceResult.Time) )
+					static FName DoPaintName(TEXT("DoPaint"));
+					if (MeshAdapter->LineTraceComponent(TraceHitResult, TraceStart, TraceEnd, FCollisionQueryParams(DoPaintName, true)))
 					{
-						BestTraceResult = TraceHitResult;
-					}
-				}	
-				
-				if (CurStaticMesh != nullptr)
-				{
-					// Reset the original collision type if we reset it. 
-					if (CachedCollisionType == ECollisionEnabled::NoCollision)
-					{
-						CurStaticMeshComponent->BodyInstance.SetCollisionEnabled(CachedCollisionType, false);
+						// Find the closest impact
+						if ((BestTraceResult.GetComponent() == nullptr) || (TraceHitResult.Time < BestTraceResult.Time))
+						{
+							BestTraceResult = TraceHitResult;
+						}
 					}
-
-					// Restore the main body setup on the mesh and recreate the physics state to update the body instance on the component.
-					CurStaticMesh->BodySetup = RestoreBodySetup;
-					CurStaticMeshComponent->RecreatePhysicsState();
 				}
 			}
 		}
 
-		if( BestTraceResult.GetActor() != NULL )
+		if (BestTraceResult.GetComponent() != NULL)
 		{
 			// If we're using texture paint, just use the best trace result we found as we currently only
 			// support painting a single mesh at a time in that mode.
-			if( FMeshPaintSettings::Get().ResourceType == EMeshPaintResource::Texture )
+			if (FMeshPaintSettings::Get().ResourceType == EMeshPaintResource::Texture)
 			{
-				PaintableActors.Add( BestTraceResult.GetActor() );
+				UMeshComponent* ComponentToPaint = CastChecked<UMeshComponent>(BestTraceResult.GetComponent());
+
+				PaintableComponents.AddUnique(ComponentToPaint);
 			}
 			else
 			{
 				FBox BrushBounds = FBox::BuildAABB( BestTraceResult.Location, FVector( BrushRadius * 1.25f, BrushRadius * 1.25f, BrushRadius * 1.25f ) );
 
-				// Vertex paint mode, so we want all valid actors overlapping the brush
-				for( int32 CurActorIndex = 0; CurActorIndex < ValidSelectedActors.Num(); ++CurActorIndex )
+				// Vertex paint mode, so we want all valid components overlapping the brush hit location
+				for (auto PotentialComponentPair : PotentialComponentMap)
 				{
-					AActor* CurValidActor = ValidSelectedActors[ CurActorIndex ];
-
-					const FBox ActorBounds = CurValidActor->GetComponentsBoundingBox( true );
+					UMeshComponent* TestComponent = PotentialComponentPair.Key;
+					const FBox ComponentBounds = TestComponent->Bounds.GetBox();
 					
-					if( ActorBounds.Intersect( BrushBounds ) )
+					if (ComponentBounds.Intersect(BrushBounds))
 					{
 						// OK, this mesh potentially overlaps the brush!
-						PaintableActors.Add( CurValidActor );
+						PaintableComponents.AddUnique(TestComponent);
 					}
 				}
 			}
 		}
 	}
 
-	bAnyPaintAbleActorsUnderCursor = (PaintableActors.Num() > 0);
+	bAnyPaintAbleActorsUnderCursor = (PaintableComponents.Num() > 0);
 
 	// Are we actually applying paint here?
 	const bool bShouldApplyPaint = bAnyPaintAbleActorsUnderCursor && ( (bIsPainting && !bVisualCueOnly) || 
@@ -1171,11 +1027,11 @@ void FEdModeMeshPaint::DoPaint( const FVector& InCameraOrigin,
 
 	// See if a Fill or PushInstanceColorsToMesh operation is requested, if so we will start an Undo/Redo transaction here
 	const bool bDoSingleFrameTransaction = ( FMeshPaintSettings::Get().ResourceType == EMeshPaintResource::VertexColors ) &&
-		PaintableActors.Num() > 0 &&
+		PaintableComponents.Num() > 0 &&
 		( ( InPaintAction == EMeshPaintAction::Fill ) || ( InPaintAction == EMeshPaintAction::PushInstanceColorsToMesh ) );
 
 	const bool bDoMultiFrameTransaction = ( FMeshPaintSettings::Get().ResourceType == EMeshPaintResource::VertexColors ) &&
-		PaintableActors.Num() > 0 &&
+		PaintableComponents.Num() > 0 &&
 		( ( InPaintAction == EMeshPaintAction::Erase ) || ( InPaintAction == EMeshPaintAction::Paint ) );
 
 	// Starts an Undo/Redo transaction with the appropriate label if we don't have any transactions in progress.
@@ -1198,134 +1054,131 @@ void FEdModeMeshPaint::DoPaint( const FVector& InCameraOrigin,
 	}
 
 	// Iterate over the selected meshes under the cursor and paint them!
-	for( int32 CurActorIndex = 0; CurActorIndex < PaintableActors.Num(); ++CurActorIndex )
+	for (UMeshComponent* MeshComponent : PaintableComponents)
 	{
-		AActor* HitActor = Cast<AActor>( PaintableActors[ CurActorIndex ] );
-		check( HitActor );
+		TSharedPtr<IMeshPaintGeometryAdapter> MeshAdapter = PotentialComponentMap.FindChecked(MeshComponent);
 
-		TArray<UMeshComponent*> MeshComponents;
-		HitActor->GetComponents<UMeshComponent>(MeshComponents);
-		check(MeshComponents.Num() > 0);
-		for(const auto& MeshComponent : MeshComponents)
-		{
-			// Brush properties
-			const float BrushDepth = BrushRadius;	// NOTE: Actually half of the total depth (like a radius)
-			const float BrushFalloffAmount = FMeshPaintSettings::Get().BrushFalloffAmount;
-			const FLinearColor BrushColor = ((InPaintAction == EMeshPaintAction::Paint) || (InPaintAction == EMeshPaintAction::Fill))? FMeshPaintSettings::Get().PaintColor : FMeshPaintSettings::Get().EraseColor;
+		// Brush properties
+		const float BrushDepth = BrushRadius;	// NOTE: Actually half of the total depth (like a radius)
+		const float BrushFalloffAmount = FMeshPaintSettings::Get().BrushFalloffAmount;
+		const FLinearColor BrushColor = ((InPaintAction == EMeshPaintAction::Paint) || (InPaintAction == EMeshPaintAction::Fill))? FMeshPaintSettings::Get().PaintColor : FMeshPaintSettings::Get().EraseColor;
 
-			// NOTE: We square the brush strength to maximize slider precision in the low range
-			const float BrushStrength =
-				FMeshPaintSettings::Get().BrushStrength * FMeshPaintSettings::Get().BrushStrength *
-				InStrengthScale;
+		// NOTE: We square the brush strength to maximize slider precision in the low range
+		const float BrushStrength =
+			FMeshPaintSettings::Get().BrushStrength * FMeshPaintSettings::Get().BrushStrength *
+			InStrengthScale;
 
-			// Display settings
-			const float VisualBiasDistance = 0.15f;
-			const float NormalLineSize( BrushRadius * 0.35f );	// Make the normal line length a function of brush size
-			const FLinearColor NormalLineColor( 0.3f, 1.0f, 0.3f );
-			const FLinearColor BrushCueColor = bIsPainting ? FLinearColor( 1.0f, 1.0f, 0.3f ) : FLinearColor( 0.3f, 1.0f, 0.3f );
-			const FLinearColor InnerBrushCueColor = bIsPainting ? FLinearColor( 0.5f, 0.5f, 0.1f ) : FLinearColor( 0.1f, 0.5f, 0.1f );
+		// Display settings
+		const float VisualBiasDistance = 0.15f;
+		const float NormalLineSize( BrushRadius * 0.35f );	// Make the normal line length a function of brush size
+		const FLinearColor NormalLineColor( 0.3f, 1.0f, 0.3f );
+		const FLinearColor BrushCueColor = bIsPainting ? FLinearColor( 1.0f, 1.0f, 0.3f ) : FLinearColor( 0.3f, 1.0f, 0.3f );
+		const FLinearColor InnerBrushCueColor = bIsPainting ? FLinearColor( 0.5f, 0.5f, 0.1f ) : FLinearColor( 0.1f, 0.5f, 0.1f );
 
-			FVector BrushXAxis, BrushYAxis;
-			BestTraceResult.Normal.FindBestAxisVectors( BrushXAxis, BrushYAxis );
-			const FVector BrushVisualPosition = BestTraceResult.Location + BestTraceResult.Normal * VisualBiasDistance;
+		FVector BrushXAxis, BrushYAxis;
+		BestTraceResult.Normal.FindBestAxisVectors( BrushXAxis, BrushYAxis );
+		const FVector BrushVisualPosition = BestTraceResult.Location + BestTraceResult.Normal * VisualBiasDistance;
 
 
-			// Precache model -> world transform
-			const FMatrix ComponentToWorldMatrix = MeshComponent->ComponentToWorld.ToMatrixWithScale();
+		// Precache model -> world transform
+		const FMatrix ComponentToWorldMatrix = MeshComponent->ComponentToWorld.ToMatrixWithScale();
 
 
-			// Compute the camera position in actor space.  We need this later to check for
-			// backfacing triangles.
-			const FVector ComponentSpaceCameraPosition( ComponentToWorldMatrix.InverseTransformPosition( InCameraOrigin ) );
-			const FVector ComponentSpaceBrushPosition( ComponentToWorldMatrix.InverseTransformPosition( BestTraceResult.Location ) );
+		// Compute the camera position in actor space.  We need this later to check for
+		// backfacing triangles.
+		const FVector ComponentSpaceCameraPosition( ComponentToWorldMatrix.InverseTransformPosition( InCameraOrigin ) );
+		const FVector ComponentSpaceBrushPosition( ComponentToWorldMatrix.InverseTransformPosition( BestTraceResult.Location ) );
 		
-			// @todo MeshPaint: Input vector doesn't work well with non-uniform scale
-			const float ComponentSpaceBrushRadius = ComponentToWorldMatrix.InverseTransformVector( FVector( BrushRadius, 0.0f, 0.0f ) ).Size();
-			const float ComponentSpaceSquaredBrushRadius = ComponentSpaceBrushRadius * ComponentSpaceBrushRadius;
+		// @todo MeshPaint: Input vector doesn't work well with non-uniform scale
+		const float ComponentSpaceBrushRadius = ComponentToWorldMatrix.InverseTransformVector( FVector( BrushRadius, 0.0f, 0.0f ) ).Size();
+		const float ComponentSpaceSquaredBrushRadius = ComponentSpaceBrushRadius * ComponentSpaceBrushRadius;
 
 
-			if( PDI != NULL )
-			{
-				// Draw brush circle
-				const int32 NumCircleSides = 64;
-				DrawCircle( PDI, BrushVisualPosition, BrushXAxis, BrushYAxis, BrushCueColor, BrushRadius, NumCircleSides, SDPG_World );
+		if( PDI != NULL )
+		{
+			// Draw brush circle
+			const int32 NumCircleSides = 64;
+			DrawCircle( PDI, BrushVisualPosition, BrushXAxis, BrushYAxis, BrushCueColor, BrushRadius, NumCircleSides, SDPG_World );
 
-				// Also draw the inner brush radius
-				const float InnerBrushRadius = BrushRadius - BrushFalloffAmount * BrushRadius;
-				DrawCircle( PDI, BrushVisualPosition, BrushXAxis, BrushYAxis, InnerBrushCueColor, InnerBrushRadius, NumCircleSides, SDPG_World );
+			// Also draw the inner brush radius
+			const float InnerBrushRadius = BrushRadius - BrushFalloffAmount * BrushRadius;
+			DrawCircle( PDI, BrushVisualPosition, BrushXAxis, BrushYAxis, InnerBrushCueColor, InnerBrushRadius, NumCircleSides, SDPG_World );
 
-				// If we just started painting then also draw a little brush effect
-				if( bIsPainting )
-				{
-					const float EffectDuration = 0.2f;
+			// If we just started painting then also draw a little brush effect
+			if( bIsPainting )
+			{
+				const float EffectDuration = 0.2f;
 
-					const double CurTime = FPlatformTime::Seconds();
-					const float TimeSinceStartedPainting = (float)( CurTime - PaintingStartTime );
-					if( TimeSinceStartedPainting <= EffectDuration )
+				const double CurTime = FPlatformTime::Seconds();
+				const float TimeSinceStartedPainting = (float)( CurTime - PaintingStartTime );
+				if( TimeSinceStartedPainting <= EffectDuration )
+				{
+					// Invert the effect if we're currently erasing
+					float EffectAlpha = TimeSinceStartedPainting / EffectDuration;
+					if( InPaintAction == EMeshPaintAction::Erase )
 					{
-						// Invert the effect if we're currently erasing
-						float EffectAlpha = TimeSinceStartedPainting / EffectDuration;
-						if( InPaintAction == EMeshPaintAction::Erase )
-						{
-							EffectAlpha = 1.0f - EffectAlpha;
-						}
-
-						const FLinearColor EffectColor( 0.1f + EffectAlpha * 0.4f, 0.1f + EffectAlpha * 0.4f, 0.1f + EffectAlpha * 0.4f );
-						const float EffectRadius = BrushRadius * EffectAlpha * EffectAlpha;	// Squared curve here (looks more interesting)
-						DrawCircle( PDI, BrushVisualPosition, BrushXAxis, BrushYAxis, EffectColor, EffectRadius, NumCircleSides, SDPG_World );
+						EffectAlpha = 1.0f - EffectAlpha;
 					}
-				}
 
-				// Draw trace surface normal
-				const FVector NormalLineEnd( BrushVisualPosition + BestTraceResult.Normal * NormalLineSize );
-				PDI->DrawLine( BrushVisualPosition, NormalLineEnd, NormalLineColor, SDPG_World );
+					const FLinearColor EffectColor( 0.1f + EffectAlpha * 0.4f, 0.1f + EffectAlpha * 0.4f, 0.1f + EffectAlpha * 0.4f );
+					const float EffectRadius = BrushRadius * EffectAlpha * EffectAlpha;	// Squared curve here (looks more interesting)
+					DrawCircle( PDI, BrushVisualPosition, BrushXAxis, BrushYAxis, EffectColor, EffectRadius, NumCircleSides, SDPG_World );
+				}
 			}
 
+			// Draw trace surface normal
+			const FVector NormalLineEnd( BrushVisualPosition + BestTraceResult.Normal * NormalLineSize );
+			PDI->DrawLine( BrushVisualPosition, NormalLineEnd, NormalLineColor, SDPG_World );
+		}
 
 
-			// Mesh paint settings
-			FMeshPaintParameters Params;
-			{
-				Params.PaintMode = FMeshPaintSettings::Get().PaintMode;
-				Params.PaintAction = InPaintAction;
-				Params.BrushPosition = BestTraceResult.Location;
-				Params.BrushNormal = BestTraceResult.Normal;
-				Params.BrushColor = BrushColor;
-				Params.SquaredBrushRadius = BrushRadius * BrushRadius;
-				Params.BrushRadialFalloffRange = BrushFalloffAmount * BrushRadius;
-				Params.InnerBrushRadius = BrushRadius - Params.BrushRadialFalloffRange;
-				Params.BrushDepth = BrushDepth;
-				Params.BrushDepthFalloffRange = BrushFalloffAmount * BrushDepth;
-				Params.InnerBrushDepth = BrushDepth - Params.BrushDepthFalloffRange;
-				Params.BrushStrength = BrushStrength;
-				Params.BrushToWorldMatrix = FMatrix( BrushXAxis, BrushYAxis, Params.BrushNormal, Params.BrushPosition );
-				Params.InverseBrushToWorldMatrix = Params.BrushToWorldMatrix.InverseFast();
-				Params.bWriteRed = FMeshPaintSettings::Get().bWriteRed;
-				Params.bWriteGreen = FMeshPaintSettings::Get().bWriteGreen;
-				Params.bWriteBlue = FMeshPaintSettings::Get().bWriteBlue;
-				Params.bWriteAlpha = FMeshPaintSettings::Get().bWriteAlpha;
-				Params.TotalWeightCount = FMeshPaintSettings::Get().TotalWeightCount;
-
-				// Select texture weight index based on whether or not we're painting or erasing
-				{
-					const int32 PaintWeightIndex = 
-						( InPaintAction == EMeshPaintAction::Paint ) ? FMeshPaintSettings::Get().PaintWeightIndex : FMeshPaintSettings::Get().EraseWeightIndex;
 
-					// Clamp the weight index to fall within the total weight count
-					Params.PaintWeightIndex = FMath::Clamp( PaintWeightIndex, 0, Params.TotalWeightCount - 1 );
-				}
+		// Mesh paint settings
+		FMeshPaintParameters Params;
+		{
+			Params.PaintMode = FMeshPaintSettings::Get().PaintMode;
+			Params.PaintAction = InPaintAction;
+			Params.BrushPosition = BestTraceResult.Location;
+			Params.BrushNormal = BestTraceResult.Normal;
+			Params.BrushColor = BrushColor;
+			Params.SquaredBrushRadius = BrushRadius * BrushRadius;
+			Params.BrushRadialFalloffRange = BrushFalloffAmount * BrushRadius;
+			Params.InnerBrushRadius = BrushRadius - Params.BrushRadialFalloffRange;
+			Params.BrushDepth = BrushDepth;
+			Params.BrushDepthFalloffRange = BrushFalloffAmount * BrushDepth;
+			Params.InnerBrushDepth = BrushDepth - Params.BrushDepthFalloffRange;
+			Params.BrushStrength = BrushStrength;
+			Params.BrushToWorldMatrix = FMatrix( BrushXAxis, BrushYAxis, Params.BrushNormal, Params.BrushPosition );
+			Params.InverseBrushToWorldMatrix = Params.BrushToWorldMatrix.InverseFast();
+			Params.bWriteRed = FMeshPaintSettings::Get().bWriteRed;
+			Params.bWriteGreen = FMeshPaintSettings::Get().bWriteGreen;
+			Params.bWriteBlue = FMeshPaintSettings::Get().bWriteBlue;
+			Params.bWriteAlpha = FMeshPaintSettings::Get().bWriteAlpha;
+			Params.TotalWeightCount = FMeshPaintSettings::Get().TotalWeightCount;
+
+			// Select texture weight index based on whether or not we're painting or erasing
+			{
+				const int32 PaintWeightIndex = 
+					( InPaintAction == EMeshPaintAction::Paint ) ? FMeshPaintSettings::Get().PaintWeightIndex : FMeshPaintSettings::Get().EraseWeightIndex;
 
-				// @todo MeshPaint: Ideally we would default to: TexturePaintingCurrentMeshComponent->StaticMesh->LightMapCoordinateIndex
-				//		Or we could indicate in the GUI which channel is the light map set (button to set it?)
-				Params.UVChannel = FMeshPaintSettings::Get().UVChannel;
+				// Clamp the weight index to fall within the total weight count
+				Params.PaintWeightIndex = FMath::Clamp( PaintWeightIndex, 0, Params.TotalWeightCount - 1 );
 			}
 
-			if (FMeshPaintSettings::Get().ResourceType == EMeshPaintResource::VertexColors)
+			// @todo MeshPaint: Ideally we would default to: TexturePaintingCurrentMeshComponent->StaticMesh->LightMapCoordinateIndex
+			//		Or we could indicate in the GUI which channel is the light map set (button to set it?)
+			Params.UVChannel = FMeshPaintSettings::Get().UVChannel;
+		}
+
+		if (FMeshPaintSettings::Get().ResourceType == EMeshPaintResource::VertexColors)
+		{
+			if (MeshAdapter->SupportsVertexPaint())
 			{
 				if (UStaticMeshComponent* StaticMeshComponent = Cast<UStaticMeshComponent>(MeshComponent))
 				{
 					if (UStaticMesh* StaticMesh = StaticMeshComponent->StaticMesh)
 					{
+						//@TODO: MESHPAINT: Direct assumptions about StaticMeshComponent
 						check(StaticMesh->GetNumLODs() > PaintingMeshLODIndex);
 						FStaticMeshLODResources& LODModel = StaticMeshComponent->StaticMesh->RenderData->LODResources[PaintingMeshLODIndex];
 
@@ -1334,14 +1187,12 @@ void FEdModeMeshPaint::DoPaint( const FVector& InCameraOrigin,
 					}
 				}
 			}
-			else
-			{
-				// Painting textures
-				PaintMeshTexture(MeshComponent, Params, bShouldApplyPaint, ComponentSpaceCameraPosition, ComponentToWorldMatrix, ComponentSpaceSquaredBrushRadius, ComponentSpaceBrushPosition);
-			}
 		}
-		
-
+		else
+		{
+			// Painting textures
+			PaintMeshTexture(MeshComponent, Params, bShouldApplyPaint, ComponentSpaceCameraPosition, ComponentToWorldMatrix, ComponentSpaceSquaredBrushRadius, ComponentSpaceBrushPosition, *MeshAdapter);
+		}
 	}
 
 	// Ends an Undo/Redo transaction, but only for Fill or PushInstanceColorsToMesh operations. Multi frame transactions will end when the user stops painting.
@@ -1778,7 +1629,7 @@ void FEdModeMeshPaint::PaintMeshVertices(
 
 
 /** Paints mesh texture */
-void FEdModeMeshPaint::PaintMeshTexture( UMeshComponent* MeshComponent, const FMeshPaintParameters& Params, const bool bShouldApplyPaint, const FVector& ComponentSpaceCameraPosition, const FMatrix& ComponentToWorldMatrix, const float ComponentSpaceSquaredBrushRadius, const FVector& ComponentSpaceBrushPosition )
+void FEdModeMeshPaint::PaintMeshTexture( UMeshComponent* MeshComponent, const FMeshPaintParameters& Params, const bool bShouldApplyPaint, const FVector& ComponentSpaceCameraPosition, const FMatrix& ComponentToWorldMatrix, const float ComponentSpaceSquaredBrushRadius, const FVector& ComponentSpaceBrushPosition, const class IMeshPaintGeometryAdapter& GeometryInfo )
 {
 	UTexture2D* TargetTexture2D = GetSelectedTexture();
 	
@@ -1791,13 +1642,6 @@ void FEdModeMeshPaint::PaintMeshTexture( UMeshComponent* MeshComponent, const FM
 	const bool bOnlyFrontFacing = FMeshPaintSettings::Get().bOnlyFrontFacingTriangles;
 	if( bShouldApplyPaint )
 	{
-		// Create the geometry adaptor
-		TSharedPtr<IMeshPaintGeometryAdaptor> GeometryInfo = PRIVATE_ConstructMeshPaintGeometryAdaptor(MeshComponent, PaintingMeshLODIndex, Params.UVChannel);
-		if (!GeometryInfo.IsValid())
-		{
-			return;
-		}
-
 		// @todo MeshPaint: Use a spatial database to reduce the triangle set here (kdop)
 		UStaticMeshComponent* StaticMeshComponent = Cast<UStaticMeshComponent>(MeshComponent);
 
@@ -1962,7 +1806,7 @@ void FEdModeMeshPaint::PaintMeshTexture( UMeshComponent* MeshComponent, const FM
 
 			if( TexturePaintingCurrentMeshComponent != NULL )
 			{
-				PaintTexture( Params, InfluencedTriangles, ComponentToWorldMatrix, *GeometryInfo );
+				PaintTexture( Params, InfluencedTriangles, ComponentToWorldMatrix, GeometryInfo );
 			}
 		}
 	}
@@ -2136,7 +1980,7 @@ void FEdModeMeshPaint::StartPaintingTexture( UMeshComponent* InMeshComponent )
 void FEdModeMeshPaint::PaintTexture( const FMeshPaintParameters& InParams,
 									 const TArray< int32 >& InInfluencedTriangles,
 									 const FMatrix& InComponentToWorldMatrix,
-									 const IMeshPaintGeometryAdaptor& GeometryInfo)
+									 const IMeshPaintGeometryAdapter& GeometryInfo)
 {
 	// We bail early if there are no influenced triangles
 	if( InInfluencedTriangles.Num() <= 0 )
diff --git a/Engine/Source/Editor/MeshPaint/Private/MeshPaintModule.cpp b/Engine/Source/Editor/MeshPaint/Private/MeshPaintModule.cpp
--- a/Engine/Source/Editor/MeshPaint/Private/MeshPaintModule.cpp
+++ b/Engine/Source/Editor/MeshPaint/Private/MeshPaintModule.cpp
@@ -2,6 +2,8 @@
 
 #include "MeshPaintPrivatePCH.h"
 #include "MeshPaintEdMode.h"
+#include "MeshPaintAdapterFactory.h"
+#include "MeshPaintStaticMeshAdapter.h"
 
 class FMeshPaintModule : public IMeshPaintModule
 {
@@ -18,6 +20,10 @@ public:
 			FSlateIcon(FEditorStyle::GetStyleSetName(), "LevelEditor.MeshPaintMode", "LevelEditor.MeshPaintMode.Small"),
 			true, 200
 			);
+
+
+		TSharedRef<FMeshPaintGeometryAdapterForStaticMeshesFactory> StaticMeshAdapter(new FMeshPaintGeometryAdapterForStaticMeshesFactory);
+		RegisterGeometryAdapterFactory(StaticMeshAdapter);
 	}
 
 	/**
@@ -27,6 +33,16 @@ public:
 	{
 		FEditorModeRegistry::Get().UnregisterMode(FBuiltinEditorModes::EM_MeshPaint);
 	}
+
+	virtual void RegisterGeometryAdapterFactory(TSharedRef<IMeshPaintGeometryAdapterFactory> Factory) override
+	{
+		FMeshPaintAdapterFactory::FactoryList.Add(Factory);
+	}
+
+	virtual void UnregisterGeometryAdapterFactory(TSharedRef<IMeshPaintGeometryAdapterFactory> Factory) override
+	{
+		FMeshPaintAdapterFactory::FactoryList.Remove(Factory);
+	}
 };
 
 IMPLEMENT_MODULE( FMeshPaintModule, MeshPaint );
diff --git a/Engine/Source/Editor/MeshPaint/Private/MeshPaintStaticMeshAdapter.cpp b/Engine/Source/Editor/MeshPaint/Private/MeshPaintStaticMeshAdapter.cpp
new file mode 100644
--- /dev/null
+++ b/Engine/Source/Editor/MeshPaint/Private/MeshPaintStaticMeshAdapter.cpp
@@ -0,0 +1,174 @@
+// Copyright 1998-2015 Epic Games, Inc. All Rights Reserved.
+
+#include "MeshPaintPrivatePCH.h"
+#include "MeshPaintStaticMeshAdapter.h"
+// 
+// #include "Materials/MaterialExpressionTextureBase.h"
+// #include "Materials/MaterialExpressionTextureSample.h"
+// #include "Materials/MaterialExpressionTextureCoordinate.h"
+// #include "Materials/MaterialExpressionTextureSampleParameter.h"
+// 
+#include "StaticMeshResources.h"
+#include "MeshPaintEdMode.h"
+// #include "Factories.h"
+// #include "ScopedTransaction.h"
+// #include "MeshPaintRendering.h"
+// #include "ImageUtils.h"
+// #include "Editor/UnrealEd/Public/Toolkits/ToolkitManager.h"
+// #include "RawMesh.h"
+// #include "Editor/UnrealEd/Public/ObjectTools.h"
+// #include "AssetToolsModule.h"
+// #include "EditorSupportDelegates.h"
+// 
+// //Slate dependencies
+// #include "Editor/LevelEditor/Public/LevelEditor.h"
+// #include "Editor/LevelEditor/Public/SLevelViewport.h"
+// #include "MessageLog.h"
+// 
+// 
+// #include "Runtime/Engine/Classes/PhysicsEngine/BodySetup.h"
+// #include "SMeshPaint.h"
+// #include "ComponentReregisterContext.h"
+// #include "CanvasTypes.h"
+// #include "Engine/Selection.h"
+// #include "Engine/TextureRenderTarget2D.h"
+// #include "EngineUtils.h"
+// #include "Engine/StaticMeshActor.h"
+// #include "Materials/MaterialInstanceConstant.h"
+
+//////////////////////////////////////////////////////////////////////////
+// FMeshPaintGeometryAdapterForStaticMeshes
+
+bool FMeshPaintGeometryAdapterForStaticMeshes::Construct(UMeshComponent* InComponent, int32 InPaintingMeshLODIndex, int32 InUVChannelIndex)
+{
+	StaticMeshComponent = Cast<UStaticMeshComponent>(InComponent);
+	if (StaticMeshComponent != nullptr)
+	{
+		if (StaticMeshComponent->StaticMesh != nullptr)
+		{
+			if (InPaintingMeshLODIndex < StaticMeshComponent->StaticMesh->GetNumLODs())
+			{
+				LODModel = &(StaticMeshComponent->StaticMesh->RenderData->LODResources[InPaintingMeshLODIndex]);
+				UVChannelIndex = InUVChannelIndex;
+				Indices = LODModel->IndexBuffer.GetArrayView();
+
+				return (UVChannelIndex >= 0) && (UVChannelIndex < (int32)LODModel->VertexBuffer.GetNumTexCoords());
+			}
+		}
+	}
+
+	return false;
+}
+
+int32 FMeshPaintGeometryAdapterForStaticMeshes::GetNumTexCoords() const
+{
+	return LODModel->VertexBuffer.GetNumTexCoords();
+}
+
+void FMeshPaintGeometryAdapterForStaticMeshes::GetTriangleInfo(int32 TriIndex, FTexturePaintTriangleInfo& OutTriInfo) const
+{
+	// Grab the vertex indices and points for this triangle
+	for (int32 TriVertexNum = 0; TriVertexNum < 3; ++TriVertexNum)
+	{
+		const int32 VertexIndex = Indices[TriIndex * 3 + TriVertexNum];
+		OutTriInfo.TriVertices[TriVertexNum] = LODModel->PositionVertexBuffer.VertexPosition(VertexIndex);
+		OutTriInfo.TriUVs[TriVertexNum] = LODModel->VertexBuffer.GetVertexUV(VertexIndex, UVChannelIndex);
+	}
+}
+
+bool FMeshPaintGeometryAdapterForStaticMeshes::LineTraceComponent(struct FHitResult& OutHit, const FVector Start, const FVector End, const struct FCollisionQueryParams& Params)
+{
+	check(StaticMeshComponent);
+	check(StaticMeshComponent->StaticMesh);
+	UStaticMesh* CurStaticMesh = StaticMeshComponent->StaticMesh;
+
+	ECollisionEnabled::Type CachedCollisionType = ECollisionEnabled::NoCollision;
+	UBodySetup* RestoreBodySetup = nullptr;
+
+	//@TODO: This shouldn't be static
+	/** A mapping of static meshes to body setups that were built for static meshes */
+	static TMap<TWeakObjectPtr<UStaticMesh>, TWeakObjectPtr<UBodySetup>> StaticMeshToTempBodySetup;
+
+
+	if (CurStaticMesh != nullptr)
+	{
+		// Get a temporary body setup the has fully detailed collision for the line traces below.
+		TWeakObjectPtr<UBodySetup>* FindBodySetupPtr = StaticMeshToTempBodySetup.Find(CurStaticMesh);
+		TWeakObjectPtr<UBodySetup> CollideAllBodySetup;
+		if (FindBodySetupPtr && FindBodySetupPtr->IsValid())
+		{
+			// Existing temporary body setup for this mesh.
+			CollideAllBodySetup = *FindBodySetupPtr;
+		}
+		else
+		{
+			// No existing body setup in the cache map - create one from the mesh's main body setup.
+			UBodySetup* TempBodySetupRaw = DuplicateObject<UBodySetup>(CurStaticMesh->BodySetup, CurStaticMesh);
+
+			// Set collide all flag so that the body creates physics meshes using ALL elements from the mesh not just the collision mesh.
+			TempBodySetupRaw->bMeshCollideAll = true;
+
+			// This forces it to recreate the physics mesh.
+			TempBodySetupRaw->InvalidatePhysicsData();
+
+			// Force it to use high detail tri-mesh for collisions.
+			TempBodySetupRaw->CollisionTraceFlag = CTF_UseComplexAsSimple;
+			TempBodySetupRaw->AggGeom.ConvexElems.Empty();
+
+			CollideAllBodySetup = TempBodySetupRaw;
+
+			// cache the body setup (remove existing entry for this mesh if is one - it must be an invalid weak ptr).
+			StaticMeshToTempBodySetup.Remove(CurStaticMesh);
+			StaticMeshToTempBodySetup.Add(CurStaticMesh, CollideAllBodySetup);
+		}
+
+		// Force the collision type to not be 'NoCollision' without it the line trace will always fail. 
+		CachedCollisionType = StaticMeshComponent->BodyInstance.GetCollisionEnabled();
+		if (CachedCollisionType == ECollisionEnabled::NoCollision)
+		{
+			StaticMeshComponent->BodyInstance.SetCollisionEnabled(ECollisionEnabled::QueryOnly, false);
+		}
+
+		// Swap the main and temp body setup on the mesh and recreate the physics state to update the body instance on the component.
+		RestoreBodySetup = CurStaticMesh->BodySetup;
+		CurStaticMesh->BodySetup = CollideAllBodySetup.Get();
+		StaticMeshComponent->RecreatePhysicsState();
+	}
+
+	// Ray trace
+	bool bHitSomething = StaticMeshComponent->LineTraceComponent(OutHit, Start, End, Params);
+
+	if (CurStaticMesh != nullptr)
+	{
+		// Reset the original collision type if we reset it. 
+		if (CachedCollisionType == ECollisionEnabled::NoCollision)
+		{
+			StaticMeshComponent->BodyInstance.SetCollisionEnabled(CachedCollisionType, false);
+		}
+
+		// Restore the main body setup on the mesh and recreate the physics state to update the body instance on the component.
+		CurStaticMesh->BodySetup = RestoreBodySetup;
+		StaticMeshComponent->RecreatePhysicsState();
+	}
+
+	return bHitSomething;
+}
+
+//////////////////////////////////////////////////////////////////////////
+// FMeshPaintGeometryAdapterForStaticMeshesFactory
+
+TSharedPtr<IMeshPaintGeometryAdapter> FMeshPaintGeometryAdapterForStaticMeshesFactory::Construct(class UMeshComponent* InComponent, int32 InPaintingMeshLODIndex, int32 InUVChannelIndex) const
+{
+	UStaticMeshComponent* StaticMeshComponent = Cast<UStaticMeshComponent>(InComponent);
+	if (StaticMeshComponent != nullptr)
+	{
+		if (StaticMeshComponent->StaticMesh != nullptr)
+		{
+			TSharedRef<FMeshPaintGeometryAdapterForStaticMeshes> Result = MakeShareable(new FMeshPaintGeometryAdapterForStaticMeshes());
+			Result->Construct(InComponent, InPaintingMeshLODIndex, InUVChannelIndex);
+			return Result;
+		}
+	}
+
+	return nullptr;
+}
diff --git a/Engine/Source/Editor/MeshPaint/Private/MeshPaintStaticMeshAdapter.h b/Engine/Source/Editor/MeshPaint/Private/MeshPaintStaticMeshAdapter.h
new file mode 100644
--- /dev/null
+++ b/Engine/Source/Editor/MeshPaint/Private/MeshPaintStaticMeshAdapter.h
@@ -0,0 +1,34 @@
+// Copyright 1998-2015 Epic Games, Inc. All Rights Reserved.
+
+#pragma once
+
+#include "MeshPaintModule.h"
+
+//////////////////////////////////////////////////////////////////////////
+// FMeshPaintGeometryAdapterForStaticMeshes
+
+class FMeshPaintGeometryAdapterForStaticMeshes : public IMeshPaintGeometryAdapter
+{
+public:
+	virtual bool Construct(UMeshComponent* InComponent, int32 InPaintingMeshLODIndex, int32 InUVChannelIndex) override;
+	virtual int32 GetNumTexCoords() const override;
+	virtual void GetTriangleInfo(int32 TriIndex, struct FTexturePaintTriangleInfo& OutTriInfo) const override;
+	virtual bool SupportsTexturePaint() const override { return true; }
+	virtual bool SupportsVertexPaint() const override { return true; }
+	virtual bool LineTraceComponent(struct FHitResult& OutHit, const FVector Start, const FVector End, const struct FCollisionQueryParams& Params) override;
+
+protected:
+	UStaticMeshComponent* StaticMeshComponent;
+	FStaticMeshLODResources* LODModel;
+	int32 UVChannelIndex;
+	FIndexArrayView Indices;
+};
+
+//////////////////////////////////////////////////////////////////////////
+// FMeshPaintGeometryAdapterForStaticMeshesFactory
+
+class FMeshPaintGeometryAdapterForStaticMeshesFactory : public IMeshPaintGeometryAdapterFactory
+{
+public:
+	virtual TSharedPtr<IMeshPaintGeometryAdapter> Construct(class UMeshComponent* InComponent, int32 InPaintingMeshLODIndex, int32 InUVChannelIndex) const override;
+};
diff --git a/Engine/Source/Editor/MeshPaint/Public/MeshPaintEdMode.h b/Engine/Source/Editor/MeshPaint/Public/MeshPaintEdMode.h
--- a/Engine/Source/Editor/MeshPaint/Public/MeshPaintEdMode.h
+++ b/Engine/Source/Editor/MeshPaint/Public/MeshPaintEdMode.h
@@ -21,6 +21,13 @@ namespace EMeshPaintResource
 }
 
 
+struct FTexturePaintTriangleInfo
+{
+	FVector TriVertices[3];
+	FVector2D TrianglePoints[3];
+	FVector2D TriUVs[3];
+};
+
 /** Mesh paint mode */
 namespace EMeshPaintMode
 {
@@ -341,13 +348,6 @@ public:
 		}
 	};
 
-	struct FTexturePaintTriangleInfo
-	{
-		FVector TriVertices[ 3 ];
-		FVector2D TrianglePoints[ 3 ];
-		FVector2D TriUVs[ 3 ];
-	};
-
 	/** Constructor */
 	FEdModeMeshPaint();
 
@@ -633,8 +633,8 @@ private:
 	void PaintMeshVertices( UStaticMeshComponent* StaticMeshComponent, const FMeshPaintParameters& Params, const bool bShouldApplyPaint, FStaticMeshLODResources& LODModel, const FVector& ActorSpaceCameraPosition, const FMatrix& ActorToWorldMatrix, FPrimitiveDrawInterface* PDI, const float VisualBiasDistance );
 
 	/** Paints mesh texture */
-	void PaintMeshTexture( UMeshComponent* MeshComponent, const FMeshPaintParameters& Params, const bool bShouldApplyPaint, const FVector& ActorSpaceCameraPosition, const FMatrix& ActorToWorldMatrix, const float ActorSpaceSquaredBrushRadius, const FVector& ActorSpaceBrushPosition );
-	
+	void PaintMeshTexture( UMeshComponent* MeshComponent, const FMeshPaintParameters& Params, const bool bShouldApplyPaint, const FVector& ActorSpaceCameraPosition, const FMatrix& ActorToWorldMatrix, const float ActorSpaceSquaredBrushRadius, const FVector& ActorSpaceBrushPosition, const class IMeshPaintGeometryAdapter& GeometryInfo );
+
 	/** Forces real-time perspective viewports */
 	void ForceRealTimeViewports( const bool bEnable, const bool bStoreCurrentState );
 
@@ -648,7 +648,7 @@ private:
 	void PaintTexture( const FMeshPaintParameters& InParams,
 					   const TArray< int32 >& InInfluencedTriangles,
 					   const FMatrix& InActorToWorldMatrix,
-					   const class IMeshPaintGeometryAdaptor& GeometryInfo);
+					   const class IMeshPaintGeometryAdapter& GeometryInfo);
 
 	/** Finishes painting a texture */
 	void FinishPaintingTexture();
@@ -755,9 +755,6 @@ private:
 	/** Array of static meshes that we've modified */
 	TArray< UStaticMesh* > ModifiedStaticMeshes;
 
-	/** A mapping of static meshes to body setups that were built for static meshes */
-	TMap<TWeakObjectPtr<UStaticMesh>, TWeakObjectPtr<UBodySetup>> StaticMeshToTempBodySetup;
-
 	/** Which mesh LOD index we're painting */
 	// @todo MeshPaint: Allow painting on other LODs?
 	static const uint32 PaintingMeshLODIndex = 0;
@@ -865,12 +862,3 @@ public:
 	 */
 	void ImportVertexColors(FEditorModeTools* ModeTools, const FString& Filename, int32 UVIndex, int32 ImportLOD, uint8 ColorMask);
 };
-
-
-class IMeshPaintGeometryAdaptor
-{
-public:
-	virtual bool Construct(UMeshComponent* InComponent, int32 InPaintingMeshLODIndex, int32 InUVChannelIndex) = 0;
-	virtual int32 GetNumTexCoords() const = 0;
-	virtual void GetTriangleInfo(int32 TriIndex, FEdModeMeshPaint::FTexturePaintTriangleInfo& OutTriInfo) const = 0;
-};
diff --git a/Engine/Source/Editor/MeshPaint/Public/MeshPaintModule.h b/Engine/Source/Editor/MeshPaint/Public/MeshPaintModule.h
--- a/Engine/Source/Editor/MeshPaint/Public/MeshPaintModule.h
+++ b/Engine/Source/Editor/MeshPaint/Public/MeshPaintModule.h
@@ -2,6 +2,31 @@
 
 #pragma once
 
+/**
+ * Interface for a class to provide mesh painting support for a subclass of UMeshComponent
+ */
+class IMeshPaintGeometryAdapter
+{
+public:
+	virtual bool Construct(UMeshComponent* InComponent, int32 InPaintingMeshLODIndex, int32 InUVChannelIndex) = 0;
+	virtual int32 GetNumTexCoords() const = 0;
+	virtual void GetTriangleInfo(int32 TriIndex, struct FTexturePaintTriangleInfo& OutTriInfo) const = 0;
+
+	virtual bool SupportsTexturePaint() const = 0;
+	virtual bool SupportsVertexPaint() const = 0;
+	virtual bool LineTraceComponent(struct FHitResult& OutHit, const FVector Start, const FVector End, const struct FCollisionQueryParams& Params) = 0;
+	virtual ~IMeshPaintGeometryAdapter() {}
+};
+
+/**
+ * Factory for IMeshPaintGeometryAdapter
+ */
+class IMeshPaintGeometryAdapterFactory
+{
+public:
+	virtual TSharedPtr<IMeshPaintGeometryAdapter> Construct(class UMeshComponent* InComponent, int32 InPaintingMeshLODIndex, int32 InUVChannelIndex) const = 0;
+};
+
 
 /**
  * MeshPaint module interface
@@ -9,5 +34,7 @@
 class IMeshPaintModule : public IModuleInterface
 {
 public:
+	virtual void RegisterGeometryAdapterFactory(TSharedRef<IMeshPaintGeometryAdapterFactory> Factory) = 0;
+	virtual void UnregisterGeometryAdapterFactory(TSharedRef<IMeshPaintGeometryAdapterFactory> Factory) = 0;
 };
 
