Editor: Make IMeshPaintGeometryAdapter::LineTraceComponent const

diff --git a/Engine/Source/Editor/MeshPaint/Private/MeshPaintModule.cpp b/Engine/Source/Editor/MeshPaint/Private/MeshPaintModule.cpp
--- a/Engine/Source/Editor/MeshPaint/Private/MeshPaintModule.cpp
+++ b/Engine/Source/Editor/MeshPaint/Private/MeshPaintModule.cpp
@@ -22,8 +22,7 @@ public:
 			);
 
 
-		TSharedRef<FMeshPaintGeometryAdapterForStaticMeshesFactory> StaticMeshAdapter(new FMeshPaintGeometryAdapterForStaticMeshesFactory);
-		RegisterGeometryAdapterFactory(StaticMeshAdapter);
+		RegisterGeometryAdapterFactory(MakeShareable(new FMeshPaintGeometryAdapterForStaticMeshesFactory));
 	}
 
 	/**
diff --git a/Engine/Source/Editor/MeshPaint/Private/MeshPaintStaticMeshAdapter.cpp b/Engine/Source/Editor/MeshPaint/Private/MeshPaintStaticMeshAdapter.cpp
--- a/Engine/Source/Editor/MeshPaint/Private/MeshPaintStaticMeshAdapter.cpp
+++ b/Engine/Source/Editor/MeshPaint/Private/MeshPaintStaticMeshAdapter.cpp
@@ -76,7 +76,7 @@ void FMeshPaintGeometryAdapterForStaticMeshes::GetTriangleInfo(int32 TriIndex, F
 	}
 }
 
-bool FMeshPaintGeometryAdapterForStaticMeshes::LineTraceComponent(struct FHitResult& OutHit, const FVector Start, const FVector End, const struct FCollisionQueryParams& Params)
+bool FMeshPaintGeometryAdapterForStaticMeshes::LineTraceComponent(struct FHitResult& OutHit, const FVector Start, const FVector End, const struct FCollisionQueryParams& Params) const
 {
 	check(StaticMeshComponent);
 	check(StaticMeshComponent->StaticMesh);
@@ -159,14 +159,15 @@ bool FMeshPaintGeometryAdapterForStaticMeshes::LineTraceComponent(struct FHitRes
 
 TSharedPtr<IMeshPaintGeometryAdapter> FMeshPaintGeometryAdapterForStaticMeshesFactory::Construct(class UMeshComponent* InComponent, int32 InPaintingMeshLODIndex, int32 InUVChannelIndex) const
 {
-	UStaticMeshComponent* StaticMeshComponent = Cast<UStaticMeshComponent>(InComponent);
-	if (StaticMeshComponent != nullptr)
+	if (UStaticMeshComponent* StaticMeshComponent = Cast<UStaticMeshComponent>(InComponent))
 	{
 		if (StaticMeshComponent->StaticMesh != nullptr)
 		{
 			TSharedRef<FMeshPaintGeometryAdapterForStaticMeshes> Result = MakeShareable(new FMeshPaintGeometryAdapterForStaticMeshes());
-			Result->Construct(InComponent, InPaintingMeshLODIndex, InUVChannelIndex);
-			return Result;
+			if (Result->Construct(InComponent, InPaintingMeshLODIndex, InUVChannelIndex))
+			{
+				return Result;
+			}
 		}
 	}
 
diff --git a/Engine/Source/Editor/MeshPaint/Private/MeshPaintStaticMeshAdapter.h b/Engine/Source/Editor/MeshPaint/Private/MeshPaintStaticMeshAdapter.h
--- a/Engine/Source/Editor/MeshPaint/Private/MeshPaintStaticMeshAdapter.h
+++ b/Engine/Source/Editor/MeshPaint/Private/MeshPaintStaticMeshAdapter.h
@@ -15,7 +15,7 @@ public:
 	virtual void GetTriangleInfo(int32 TriIndex, struct FTexturePaintTriangleInfo& OutTriInfo) const override;
 	virtual bool SupportsTexturePaint() const override { return true; }
 	virtual bool SupportsVertexPaint() const override { return true; }
-	virtual bool LineTraceComponent(struct FHitResult& OutHit, const FVector Start, const FVector End, const struct FCollisionQueryParams& Params) override;
+	virtual bool LineTraceComponent(struct FHitResult& OutHit, const FVector Start, const FVector End, const struct FCollisionQueryParams& Params) const override;
 
 protected:
 	UStaticMeshComponent* StaticMeshComponent;
diff --git a/Engine/Source/Editor/MeshPaint/Public/MeshPaintModule.h b/Engine/Source/Editor/MeshPaint/Public/MeshPaintModule.h
--- a/Engine/Source/Editor/MeshPaint/Public/MeshPaintModule.h
+++ b/Engine/Source/Editor/MeshPaint/Public/MeshPaintModule.h
@@ -14,7 +14,7 @@ public:
 
 	virtual bool SupportsTexturePaint() const = 0;
 	virtual bool SupportsVertexPaint() const = 0;
-	virtual bool LineTraceComponent(struct FHitResult& OutHit, const FVector Start, const FVector End, const struct FCollisionQueryParams& Params) = 0;
+	virtual bool LineTraceComponent(struct FHitResult& OutHit, const FVector Start, const FVector End, const struct FCollisionQueryParams& Params) const = 0;
 	virtual ~IMeshPaintGeometryAdapter() {}
 };
 
