Editor: Start pushing the mesh texture painting brush sphere intersection code into the geometry adapter (WIP)

diff --git a/Engine/Plugins/2D/Paper2D/Source/Paper2DEditor/Private/MeshPainting/MeshPaintSpriteAdapter.cpp b/Engine/Plugins/2D/Paper2D/Source/Paper2DEditor/Private/MeshPainting/MeshPaintSpriteAdapter.cpp
--- a/Engine/Plugins/2D/Paper2D/Source/Paper2DEditor/Private/MeshPainting/MeshPaintSpriteAdapter.cpp
+++ b/Engine/Plugins/2D/Paper2D/Source/Paper2DEditor/Private/MeshPainting/MeshPaintSpriteAdapter.cpp
@@ -2,6 +2,7 @@
 
 #include "Paper2DEditorPrivatePCH.h"
 #include "MeshPaintSpriteAdapter.h"
+#include "MeshPaintEdMode.h"
 
 //////////////////////////////////////////////////////////////////////////
 // FMeshPaintSpriteAdapter
@@ -19,13 +20,27 @@ int32 FMeshPaintSpriteAdapter::GetNumTexCoords() const
 
 void FMeshPaintSpriteAdapter::GetTriangleInfo(int32 TriIndex, struct FTexturePaintTriangleInfo& OutTriInfo) const
 {
-	//@TODO: MESHPAINT: Write this
+	UPaperSprite* Sprite = SpriteComponent->GetSprite();
+	checkSlow(Sprite != nullptr);
+
+	// Grab the vertex indices and points for this triangle
+	const TArray<FVector4>& BakedPoints = Sprite->BakedRenderData;
+
+	for (int32 TriVertexNum = 0; TriVertexNum < 3; ++TriVertexNum)
+	{
+		const int32 VertexIndex = (TriIndex * 3) + TriVertexNum;
+		const FVector4& XYUV = BakedPoints[VertexIndex];
+
+		OutTriInfo.TriVertices[TriVertexNum] = (PaperAxisX * XYUV.X) + (PaperAxisY * XYUV.Y);
+		OutTriInfo.TriUVs[TriVertexNum] = FVector2D(XYUV.Z, XYUV.W);
+	}
 }
 
 
 bool FMeshPaintSpriteAdapter::LineTraceComponent(struct FHitResult& OutHit, const FVector Start, const FVector End, const struct FCollisionQueryParams& Params) const
 {
 	UPaperSprite* Sprite = SpriteComponent->GetSprite();
+	checkSlow(Sprite != nullptr);
 
 	const FTransform& ComponentToWorld = SpriteComponent->ComponentToWorld;
 
@@ -83,6 +98,20 @@ bool FMeshPaintSpriteAdapter::LineTraceComponent(struct FHitResult& OutHit, cons
 	return false;
 }
 
+void FMeshPaintSpriteAdapter::SphereIntersectTriangles(TArray<int32>& OutTriangles, const float ComponentSpaceSquaredBrushRadius, const FVector& ComponentSpaceBrushPosition) const
+{
+	UPaperSprite* Sprite = SpriteComponent->GetSprite();
+	checkSlow(Sprite != nullptr);
+
+	//@TODO: MESHPAINT: This isn't very precise..., but since the sprite is planar it shouldn't cause any actual issues with paint UX other than being suboptimal
+	const int32 NumTriangles = Sprite->BakedRenderData.Num() / 3;
+	OutTriangles.Reserve(OutTriangles.Num() + NumTriangles);
+	for (int32 TriIndex = 0; TriIndex < NumTriangles; ++TriIndex)
+	{
+		OutTriangles.Add(TriIndex);
+	}
+}
+
 //////////////////////////////////////////////////////////////////////////
 // FMeshPaintSpriteAdapterFactory
 
diff --git a/Engine/Plugins/2D/Paper2D/Source/Paper2DEditor/Private/MeshPainting/MeshPaintSpriteAdapter.h b/Engine/Plugins/2D/Paper2D/Source/Paper2DEditor/Private/MeshPainting/MeshPaintSpriteAdapter.h
--- a/Engine/Plugins/2D/Paper2D/Source/Paper2DEditor/Private/MeshPainting/MeshPaintSpriteAdapter.h
+++ b/Engine/Plugins/2D/Paper2D/Source/Paper2DEditor/Private/MeshPainting/MeshPaintSpriteAdapter.h
@@ -16,6 +16,7 @@ public:
 	virtual bool SupportsTexturePaint() const override { return true; }
 	virtual bool SupportsVertexPaint() const override { return false; }
 	virtual bool LineTraceComponent(struct FHitResult& OutHit, const FVector Start, const FVector End, const struct FCollisionQueryParams& Params) const override;
+	virtual void SphereIntersectTriangles(TArray<int32>& OutTriangles, const float ComponentSpaceSquaredBrushRadius, const FVector& ComponentSpaceBrushPosition) const override;
 
 protected:
 	UPaperSpriteComponent* SpriteComponent;
diff --git a/Engine/Source/Editor/MeshPaint/Private/MeshPaintEdMode.cpp b/Engine/Source/Editor/MeshPaint/Private/MeshPaintEdMode.cpp
--- a/Engine/Source/Editor/MeshPaint/Private/MeshPaintEdMode.cpp
+++ b/Engine/Source/Editor/MeshPaint/Private/MeshPaintEdMode.cpp
@@ -1633,23 +1633,48 @@ void FEdModeMeshPaint::PaintMeshTexture( UMeshComponent* MeshComponent, const FM
 {
 	UTexture2D* TargetTexture2D = GetSelectedTexture();
 	
-	// No reason to continue if we dont have a target texture;
-	if( TargetTexture2D == NULL )
+	// No reason to continue if we dont have a target texture or we are not painting
+	if ((TargetTexture2D == NULL) || !bShouldApplyPaint)
 	{
 		return;
 	}
 
 	const bool bOnlyFrontFacing = FMeshPaintSettings::Get().bOnlyFrontFacingTriangles;
-	if( bShouldApplyPaint )
+
+	PaintTexture2DData* TextureData = GetPaintTargetData(TargetTexture2D);
+
+	// Store info that tells us if the element material uses our target texture so we don't have to do a usestexture() call for each tri
+	TArray<bool> SectionUsesTargetTexture;
+	SectionUsesTargetTexture.AddZeroed(MeshComponent->GetNumMaterials());
+	for (int32 SectionIndex = 0; SectionIndex < SectionUsesTargetTexture.Num(); SectionIndex++)
 	{
-		// @todo MeshPaint: Use a spatial database to reduce the triangle set here (kdop)
-		UStaticMeshComponent* StaticMeshComponent = Cast<UStaticMeshComponent>(MeshComponent);
+		SectionUsesTargetTexture[SectionIndex] = false;
 
-		if ((StaticMeshComponent == nullptr) || (StaticMeshComponent->StaticMesh == nullptr))
+		// @todo MeshPaint: if LODs can use different materials/textures then this will cause us problems
+		if (UMaterialInterface* SectionMat = MeshComponent->GetMaterial(SectionIndex))
 		{
-			return;
+			SectionUsesTargetTexture[SectionIndex] |= DoesMaterialUseTexture(SectionMat, TargetTexture2D);
+
+			if ((SectionUsesTargetTexture[SectionIndex] == false) && (TextureData != NULL) && (TextureData->PaintRenderTargetTexture != NULL))
+			{
+				// If we didn't get a match on our selected texture, we'll check to see if the the material uses a
+				//  render target texture override that we put on during painting.
+				SectionUsesTargetTexture[SectionIndex] |= DoesMaterialUseTexture(SectionMat, TextureData->PaintRenderTargetTexture);
+			}
 		}
+	}
+
+	// Keep a list of front-facing triangles that are within a reasonable distance to the brush
+	static TArray< int32 > InfluencedTriangles;
 
+
+
+	// @todo MeshPaint: Use a spatial database to reduce the triangle set here (kdop)
+	UStaticMeshComponent* StaticMeshComponent = Cast<UStaticMeshComponent>(MeshComponent);
+
+	if ((StaticMeshComponent != nullptr) && (StaticMeshComponent->StaticMesh == nullptr))
+	{
+		//@TODO: MESHPAINT: Move this code to the adapter (need to determine if flushing the octree every frame is disastrous and if so how to add multi-frame state to the adapters)
 		check(StaticMeshComponent->StaticMesh->GetNumLODs() > PaintingMeshLODIndex);
 		FStaticMeshLODResources& LODModel = StaticMeshComponent->StaticMesh->RenderData->LODResources[PaintingMeshLODIndex];
 		const int32 NumSections = LODModel.Sections.Num();
@@ -1660,41 +1685,12 @@ void FEdModeMeshPaint::PaintMeshTexture( UMeshComponent* MeshComponent, const FM
 		const uint32 NumIndexBufferIndices = Indices.Num();
 		check( NumIndexBufferIndices % 3 == 0 );
 		const uint32 NumTriangles = NumIndexBufferIndices / 3;
-
-		// Keep a list of front-facing triangles that are within a reasonable distance to the brush
-		static TArray< int32 > InfluencedTriangles;
-		InfluencedTriangles.Empty( NumTriangles );
+		InfluencedTriangles.Empty(NumTriangles);
 
 		// Use a bit of distance bias to make sure that we get all of the overlapping triangles.  We
 		// definitely don't want our brush to be cut off by a hard triangle edge
 		const float SquaredRadiusBias = ComponentSpaceSquaredBrushRadius * 0.025f;
 
-		
-		PaintTexture2DData* TextureData = GetPaintTargetData( TargetTexture2D );
-
-		// Store info that tells us if the element material uses our target texture so we don't have to do a usestexture() call for each tri
-		TArray< bool > SectionUsesTargetTexture;
-		SectionUsesTargetTexture.AddZeroed( NumSections );
-		for ( int32 SectionIndex = 0; SectionIndex < NumSections; SectionIndex++ )
-		{
-			SectionUsesTargetTexture[ SectionIndex ] = false;
-
-			// @todo MeshPaint: if LODs can use different materials/textures then this will cause us problems
-			UMaterialInterface* SectionMat = StaticMeshComponent->GetMaterial( SectionIndex );
-			if( SectionMat != NULL )
-			{
-				
-				SectionUsesTargetTexture[ SectionIndex ] |=  DoesMaterialUseTexture(SectionMat,TargetTexture2D);
-				
-				if( SectionUsesTargetTexture[ SectionIndex ] == false && TextureData != NULL && TextureData->PaintRenderTargetTexture != NULL)
-				{
-					// If we didn't get a match on our selected texture, we'll check to see if the the material uses a
-					//  render target texture override that we put on during painting.
-					SectionUsesTargetTexture[ SectionIndex ] |=  DoesMaterialUseTexture(SectionMat,TextureData->PaintRenderTargetTexture);
-				}
-			}
-		}
-
 		if( (TexturePaintingStaticMeshOctree != NULL) && ((TexturePaintingCurrentMeshComponent != MeshComponent) || (TexturePaintingStaticMeshLOD != PaintingMeshLODIndex)) )
 		{
 			delete TexturePaintingStaticMeshOctree;
@@ -1765,13 +1761,11 @@ void FEdModeMeshPaint::PaintMeshTexture( UMeshComponent* MeshComponent, const FM
 				{
 					FStaticMeshSection& Section = LODModel.Sections[ SectionIndex ];
 
-
 					if( ( CurrentTri.Index >= Section.FirstIndex / 3 ) && 
 						( CurrentTri.Index < Section.FirstIndex / 3 + Section.NumTriangles ) )
 					{
-
 						// The triangle belongs to this element, now we need to check to see if the element material uses our target texture.
-						if( TargetTexture2D != NULL && SectionUsesTargetTexture[ SectionIndex ] == true)
+						if ((TargetTexture2D != NULL) && SectionUsesTargetTexture.IsValidIndex(Section.MaterialIndex) && (SectionUsesTargetTexture[Section.MaterialIndex] == true))
 						{
 							bAddTri = true;
 						}
@@ -1788,27 +1782,29 @@ void FEdModeMeshPaint::PaintMeshTexture( UMeshComponent* MeshComponent, const FM
 				}
 			}
 		}
+	}
+	else
+	{
+		// Try the adapter
+		GeometryInfo.SphereIntersectTriangles(InfluencedTriangles, ComponentSpaceSquaredBrushRadius, ComponentSpaceBrushPosition);
+	}
 
 
-		{
-
-			if( (TexturePaintingCurrentMeshComponent != NULL) && (TexturePaintingCurrentMeshComponent != MeshComponent) )
-			{
-				// Mesh has changed, so finish up with our previous texture
-				FinishPaintingTexture();
-				bIsPainting = false;
-			}
+	if ((TexturePaintingCurrentMeshComponent != NULL) && (TexturePaintingCurrentMeshComponent != MeshComponent))
+	{
+		// Mesh has changed, so finish up with our previous texture
+		FinishPaintingTexture();
+		bIsPainting = false;
+	}
 
-			if( TexturePaintingCurrentMeshComponent == NULL )
-			{
-				StartPaintingTexture( MeshComponent );
-			}
+	if (TexturePaintingCurrentMeshComponent == NULL)
+	{
+		StartPaintingTexture( MeshComponent );
+	}
 
-			if( TexturePaintingCurrentMeshComponent != NULL )
-			{
-				PaintTexture( Params, InfluencedTriangles, ComponentToWorldMatrix, GeometryInfo );
-			}
-		}
+	if (TexturePaintingCurrentMeshComponent != NULL)
+	{
+		PaintTexture( Params, InfluencedTriangles, ComponentToWorldMatrix, GeometryInfo );
 	}
 }
 
diff --git a/Engine/Source/Editor/MeshPaint/Private/MeshPaintStaticMeshAdapter.cpp b/Engine/Source/Editor/MeshPaint/Private/MeshPaintStaticMeshAdapter.cpp
--- a/Engine/Source/Editor/MeshPaint/Private/MeshPaintStaticMeshAdapter.cpp
+++ b/Engine/Source/Editor/MeshPaint/Private/MeshPaintStaticMeshAdapter.cpp
@@ -154,6 +154,11 @@ bool FMeshPaintGeometryAdapterForStaticMeshes::LineTraceComponent(struct FHitRes
 	return bHitSomething;
 }
 
+void FMeshPaintGeometryAdapterForStaticMeshes::SphereIntersectTriangles(TArray<int32>& OutTriangles, const float ComponentSpaceSquaredBrushRadius, const FVector& ComponentSpaceBrushPosition) const
+{
+	//@TODO: The code for static mesh components is still inside of the mesh painting mode due to the octree
+}
+
 //////////////////////////////////////////////////////////////////////////
 // FMeshPaintGeometryAdapterForStaticMeshesFactory
 
diff --git a/Engine/Source/Editor/MeshPaint/Private/MeshPaintStaticMeshAdapter.h b/Engine/Source/Editor/MeshPaint/Private/MeshPaintStaticMeshAdapter.h
--- a/Engine/Source/Editor/MeshPaint/Private/MeshPaintStaticMeshAdapter.h
+++ b/Engine/Source/Editor/MeshPaint/Private/MeshPaintStaticMeshAdapter.h
@@ -16,6 +16,7 @@ public:
 	virtual bool SupportsTexturePaint() const override { return true; }
 	virtual bool SupportsVertexPaint() const override { return true; }
 	virtual bool LineTraceComponent(struct FHitResult& OutHit, const FVector Start, const FVector End, const struct FCollisionQueryParams& Params) const override;
+	virtual void SphereIntersectTriangles(TArray<int32>& OutTriangles, const float ComponentSpaceSquaredBrushRadius, const FVector& ComponentSpaceBrushPosition) const override;
 
 protected:
 	UStaticMeshComponent* StaticMeshComponent;
diff --git a/Engine/Source/Editor/MeshPaint/Public/MeshPaintModule.h b/Engine/Source/Editor/MeshPaint/Public/MeshPaintModule.h
--- a/Engine/Source/Editor/MeshPaint/Public/MeshPaintModule.h
+++ b/Engine/Source/Editor/MeshPaint/Public/MeshPaintModule.h
@@ -15,6 +15,7 @@ public:
 	virtual bool SupportsTexturePaint() const = 0;
 	virtual bool SupportsVertexPaint() const = 0;
 	virtual bool LineTraceComponent(struct FHitResult& OutHit, const FVector Start, const FVector End, const struct FCollisionQueryParams& Params) const = 0;
+	virtual void SphereIntersectTriangles(TArray<int32>& OutTriangles, const float ComponentSpaceSquaredBrushRadius, const FVector& ComponentSpaceBrushPosition) const = 0;
 	virtual ~IMeshPaintGeometryAdapter() {}
 };
 
@@ -25,6 +26,7 @@ class IMeshPaintGeometryAdapterFactory
 {
 public:
 	virtual TSharedPtr<IMeshPaintGeometryAdapter> Construct(class UMeshComponent* InComponent, int32 InPaintingMeshLODIndex, int32 InUVChannelIndex) const = 0;
+	virtual ~IMeshPaintGeometryAdapterFactory() {}
 };
 
 
