Editor: Fix mesh texture painting on static mesh components

diff --git a/Engine/Source/Editor/MeshPaint/Private/MeshPaintEdMode.cpp b/Engine/Source/Editor/MeshPaint/Private/MeshPaintEdMode.cpp
--- a/Engine/Source/Editor/MeshPaint/Private/MeshPaintEdMode.cpp
+++ b/Engine/Source/Editor/MeshPaint/Private/MeshPaintEdMode.cpp
@@ -1672,7 +1672,7 @@ void FEdModeMeshPaint::PaintMeshTexture( UMeshComponent* MeshComponent, const FM
 	// @todo MeshPaint: Use a spatial database to reduce the triangle set here (kdop)
 	UStaticMeshComponent* StaticMeshComponent = Cast<UStaticMeshComponent>(MeshComponent);
 
-	if ((StaticMeshComponent != nullptr) && (StaticMeshComponent->StaticMesh == nullptr))
+	if ((StaticMeshComponent != nullptr) && (StaticMeshComponent->StaticMesh != nullptr))
 	{
 		//@TODO: MESHPAINT: Move this code to the adapter (need to determine if flushing the octree every frame is disastrous and if so how to add multi-frame state to the adapters)
 		check(StaticMeshComponent->StaticMesh->GetNumLODs() > PaintingMeshLODIndex);
