Paper2D: Improvements to texture selection for mesh paint, allowing adapters to override behavior

diff --git a/Engine/Plugins/2D/Paper2D/Source/Paper2DEditor/Paper2DEditor.Build.cs b/Engine/Plugins/2D/Paper2D/Source/Paper2DEditor/Paper2DEditor.Build.cs
--- a/Engine/Plugins/2D/Paper2D/Source/Paper2DEditor/Paper2DEditor.Build.cs
+++ b/Engine/Plugins/2D/Paper2D/Source/Paper2DEditor/Paper2DEditor.Build.cs
@@ -28,15 +28,15 @@ public class Paper2DEditor : ModuleRules
 				"Paper2D",
 				"ContentBrowser",
 				"WorkspaceMenuStructure",
-				"EditorStyle"
+				"EditorStyle",
+				"MeshPaint"
 			}
 			);
 
 		PrivateIncludePathModuleNames.AddRange(
 			new string[] {
 				"Settings",
-				"IntroTutorials",
-				"MeshPaint"
+				"IntroTutorials"
 			}
 		);
 	}
diff --git a/Engine/Plugins/2D/Paper2D/Source/Paper2DEditor/Private/MeshPainting/MeshPaintSpriteAdapter.cpp b/Engine/Plugins/2D/Paper2D/Source/Paper2DEditor/Private/MeshPainting/MeshPaintSpriteAdapter.cpp
--- a/Engine/Plugins/2D/Paper2D/Source/Paper2DEditor/Private/MeshPainting/MeshPaintSpriteAdapter.cpp
+++ b/Engine/Plugins/2D/Paper2D/Source/Paper2DEditor/Private/MeshPainting/MeshPaintSpriteAdapter.cpp
@@ -112,6 +112,28 @@ void FMeshPaintSpriteAdapter::SphereIntersectTriangles(TArray<int32>& OutTriangl
 	}
 }
 
+void FMeshPaintSpriteAdapter::QueryPaintableTextures(int32 MaterialIndex, int32& OutDefaultIndex, TArray<struct FPaintableTexture>& InOutTextureList)
+{
+	UPaperSprite* Sprite = SpriteComponent->GetSprite();
+	checkSlow(Sprite != nullptr);
+
+	// Grab the sprite texture first
+	int32 ForceIndex = INDEX_NONE;
+	if (UTexture2D* SourceTexture = Sprite->GetBakedTexture())
+	{
+		new (InOutTextureList) FPaintableTexture(SourceTexture, 0);
+		ForceIndex = InOutTextureList.Num() - 1;
+	}
+	
+	// Now ask the material
+	DefaultQueryPaintableTextures(MaterialIndex, SpriteComponent, OutDefaultIndex, InOutTextureList);
+
+	if (ForceIndex != INDEX_NONE)
+	{
+		OutDefaultIndex = ForceIndex;
+	}
+}
+
 //////////////////////////////////////////////////////////////////////////
 // FMeshPaintSpriteAdapterFactory
 
diff --git a/Engine/Plugins/2D/Paper2D/Source/Paper2DEditor/Private/MeshPainting/MeshPaintSpriteAdapter.h b/Engine/Plugins/2D/Paper2D/Source/Paper2DEditor/Private/MeshPainting/MeshPaintSpriteAdapter.h
--- a/Engine/Plugins/2D/Paper2D/Source/Paper2DEditor/Private/MeshPainting/MeshPaintSpriteAdapter.h
+++ b/Engine/Plugins/2D/Paper2D/Source/Paper2DEditor/Private/MeshPainting/MeshPaintSpriteAdapter.h
@@ -17,6 +17,7 @@ public:
 	virtual bool SupportsVertexPaint() const override { return false; }
 	virtual bool LineTraceComponent(struct FHitResult& OutHit, const FVector Start, const FVector End, const struct FCollisionQueryParams& Params) const override;
 	virtual void SphereIntersectTriangles(TArray<int32>& OutTriangles, const float ComponentSpaceSquaredBrushRadius, const FVector& ComponentSpaceBrushPosition) const override;
+	virtual void QueryPaintableTextures(int32 MaterialIndex, int32& OutDefaultIndex, TArray<struct FPaintableTexture>& InOutTextureList) override;
 
 protected:
 	UPaperSpriteComponent* SpriteComponent;
diff --git a/Engine/Source/Editor/MeshPaint/Private/MeshPaintEdMode.cpp b/Engine/Source/Editor/MeshPaint/Private/MeshPaintEdMode.cpp
--- a/Engine/Source/Editor/MeshPaint/Private/MeshPaintEdMode.cpp
+++ b/Engine/Source/Editor/MeshPaint/Private/MeshPaintEdMode.cpp
@@ -2,11 +2,6 @@
 
 #include "MeshPaintPrivatePCH.h"
 
-#include "Materials/MaterialExpressionTextureBase.h"
-#include "Materials/MaterialExpressionTextureSample.h"
-#include "Materials/MaterialExpressionTextureCoordinate.h"
-#include "Materials/MaterialExpressionTextureSampleParameter.h"
-
 #include "StaticMeshResources.h"
 #include "MeshPaintEdMode.h"
 #include "Factories.h"
@@ -61,6 +56,13 @@ public:
 	MeshPaintRendering::FMeshPaintShaderParameters ShaderParams;
 };
 
+bool DoesMeshComponentUseTexture(UMeshComponent* MeshComponent, UTexture* Texture)
+{
+	TArray<UTexture*> UsedTextures;
+	MeshComponent->GetUsedTextures(UsedTextures, EMaterialQualityLevel::High);
+
+	return UsedTextures.Contains(Texture);
+}
 
 /** Batched element parameters for texture paint shaders used for texture dilation */
 class FMeshPaintDilateBatchedElementParameters : public FBatchedElementParameters
@@ -1665,7 +1667,7 @@ void FEdModeMeshPaint::PaintMeshTexture( UMeshComponent* MeshComponent, const FM
 	}
 
 	// Keep a list of front-facing triangles that are within a reasonable distance to the brush
-	static TArray< int32 > InfluencedTriangles;
+	TArray< int32 > InfluencedTriangles;
 
 
 
@@ -1834,11 +1836,11 @@ void FEdModeMeshPaint::StartPaintingTexture( UMeshComponent* InMeshComponent )
 	UMaterialInterface* MaterialToCheck = InMeshComponent->GetMaterial(MaterialIndex);
 	while( MaterialToCheck != NULL )
 	{
-		bool bIsTextureUsed = DoesMaterialUseTexture(MaterialToCheck,Texture2D);
+		bool bIsTextureUsed = DoesMeshComponentUseTexture(InMeshComponent, Texture2D);
 
 		if( !bIsTextureUsed && TextureData != NULL && TextureData->PaintRenderTargetTexture != NULL )
 		{
-			bIsTextureUsed = DoesMaterialUseTexture(MaterialToCheck,TextureData->PaintRenderTargetTexture);
+			bIsTextureUsed = DoesMeshComponentUseTexture(InMeshComponent, TextureData->PaintRenderTargetTexture);
 		}
 
 		if( bIsTextureUsed == true && bStartedPainting == false )
@@ -4371,25 +4373,6 @@ void FImportVertexTextureHelper::ImportVertexColors(FEditorModeTools* ModeTools,
 	}
 }
 
-/** Structure used to house and compare Texture and UV channel pairs */
-struct FPaintableTexture
-{	
-	UTexture*	Texture;
-	int32		UVChannelIndex;
-
-	FPaintableTexture(UTexture* InTexture = NULL, uint32 InUVChannelIndex = 0)
-		: Texture(InTexture)
-		, UVChannelIndex(InUVChannelIndex)
-	{}
-	
-	/** Overloaded equality operator for use with TArrays Contains method. */
-	bool operator==(const FPaintableTexture& rhs) const
-	{
-		return (Texture == rhs.Texture);
-		/* && (UVChannelIndex == rhs.UVChannelIndex);*/// if we compare UVChannel we would have to duplicate the texture
-	}
-};
-
 /** 
  * Will update the list of available texture paint targets based on selection 
  */
@@ -4420,144 +4403,106 @@ void FEdModeMeshPaint::UpdateTexturePaintTargetList()
 				TArray<UMeshComponent*> MeshComponents;
 				Actor->GetComponents<UMeshComponent>(MeshComponents);
 
-				for(const auto& MeshComponent : MeshComponents)
+				int32 DefaultIndex = INDEX_NONE;
+				for (const auto& MeshComponent : MeshComponents)
 				{
-					if ( MeshComponent != NULL )
+					// Create the geometry adapter
+					TSharedPtr<IMeshPaintGeometryAdapter> MeshAdapter = FMeshPaintAdapterFactory::CreateAdapterForMesh(MeshComponent, PaintingMeshLODIndex, /*TODO: Shouldn't be part of the construction contract: FMeshPaintSettings::Get().UVChannel*/ 0);
+					if (MeshAdapter.IsValid())
 					{
 						// We already know the material we are painting on, take it off the static mesh component
 						UMaterialInterface* Material = MeshComponent->GetMaterial(MaterialIndex);
-				
-						if ( Material != NULL )
+
+						if (Material != NULL)
 						{
-							int32 DefaultIndex = INDEX_NONE;
-							FPaintableTexture PaintableTexture;
 							// Find all the unique textures used in the top material level of the selected actor materials
-		
-							const TArray<UMaterialExpression*>& Expressions = Material->GetMaterial()->Expressions;
-
-							// Only grab the textures from the top level of samples
-							for (auto ItExpressions = Expressions.CreateConstIterator(); ItExpressions; ItExpressions++)
+							int32 DefaultIndexHere = INDEX_NONE;
+							MeshAdapter->QueryPaintableTextures(MaterialIndex, /*out*/ DefaultIndexHere, /*inout*/ TexturesInSelection);
+							if (DefaultIndex == INDEX_NONE)
 							{
-								UMaterialExpressionTextureBase* TextureBase = Cast<UMaterialExpressionTextureBase>(*ItExpressions);
-								if (TextureBase != NULL && 
-									TextureBase->Texture != NULL && 
-									!TextureBase->Texture->IsNormalMap())
-								{
-									// Default UV channel to index 0. 
-									PaintableTexture = FPaintableTexture(TextureBase->Texture, 0);
-										
-									// Texture Samples can have UV's specified, check the first node for whether it has a custom UV channel set. 
-									// We only check the first as the Mesh paint mode does not support painting with UV's modified in the shader.
-									UMaterialExpressionTextureSample* TextureSample = Cast<UMaterialExpressionTextureSample>(*ItExpressions);
-									if (TextureSample != NULL)
-									{
-										UMaterialExpressionTextureCoordinate* TextureCoords = Cast<UMaterialExpressionTextureCoordinate>(TextureSample->Coordinates.Expression);
-										if (TextureCoords != NULL)
-										{
-											// Store the uv channel, this is set when the texture is selected. 
-											PaintableTexture.UVChannelIndex = TextureCoords->CoordinateIndex;
-										}
-
-										// Handle texture parameter expressions
-										UMaterialExpressionTextureSampleParameter* TextureSampleParameter = Cast<UMaterialExpressionTextureSampleParameter>(TextureSample);
-										if (TextureSampleParameter != NULL)
-										{
-											// Grab the overridden texture if it exists.  
-											Material->GetTextureParameterValue(TextureSampleParameter->ParameterName, PaintableTexture.Texture);
-										}
-									}
-
-									// note that the same texture will be added again if its UV channel differs. 
-									int32 TextureIndex = TexturesInSelection.AddUnique(PaintableTexture);
-
-									// cache the first default index, if there is no previous info this will be used as the selected texture
-									if (DefaultIndex == INDEX_NONE && TextureBase->IsDefaultMeshpaintTexture)
-									{
-										DefaultIndex = TextureIndex;
-									}
-								}
+								DefaultIndex = DefaultIndexHere;
 							}
+						}
+					}
+				}
+
+				// Generate the list of target paint textures that will be displaying in the UI
+				for( int32 TexIndex = 0; TexIndex < TexturesInSelection.Num(); TexIndex++ )
+				{
+					UTexture2D* Texture2D = Cast<UTexture2D>( TexturesInSelection[ TexIndex ].Texture );
+					int32 UVChannelIndex = TexturesInSelection[ TexIndex ].UVChannelIndex;
+					// If this is not a UTexture2D we check to see if it is a rendertarget texture
+					if( Texture2D == NULL )
+					{
+						UTextureRenderTarget2D* TextureRenderTarget2D = Cast<UTextureRenderTarget2D>( TexturesInSelection[ TexIndex ].Texture );
+						if( TextureRenderTarget2D )
+						{
+							// Since this is a rendertarget, we lookup the original texture that we overrode during the paint operation
+							Texture2D = GetOriginalTextureFromRenderTarget( TextureRenderTarget2D );
 
-							// Generate the list of target paint textures that will be displaying in the UI
-							for( int32 TexIndex = 0; TexIndex < TexturesInSelection.Num(); TexIndex++ )
+							// Since we looked up a texture via a rendertarget, it is possible that this texture already exists in our list.  If so 
+							//  we will not add it and continue processing other elements.
+							if( Texture2D != NULL && TexturesInSelection.Contains( FPaintableTexture(Texture2D, UVChannelIndex) ) )
 							{
-								UTexture2D* Texture2D = Cast<UTexture2D>( TexturesInSelection[ TexIndex ].Texture );
-								int32 UVChannelIndex = TexturesInSelection[ TexIndex ].UVChannelIndex;
-								// If this is not a UTexture2D we check to see if it is a rendertarget texture
-								if( Texture2D == NULL )
-								{
-									UTextureRenderTarget2D* TextureRenderTarget2D = Cast<UTextureRenderTarget2D>( TexturesInSelection[ TexIndex ].Texture );
-									if( TextureRenderTarget2D )
-									{
-										// Since this is a rendertarget, we lookup the original texture that we overrode during the paint operation
-										Texture2D = GetOriginalTextureFromRenderTarget( TextureRenderTarget2D );
-
-										// Since we looked up a texture via a rendertarget, it is possible that this texture already exists in our list.  If so 
-										//  we will not add it and continue processing other elements.
-										if( Texture2D != NULL && TexturesInSelection.Contains( FPaintableTexture(Texture2D, UVChannelIndex) ) )
-										{
-											continue;
-										}
-									}
-								}
+								continue;
+							}
+						}
+					}
 
-								if( Texture2D != NULL )
-								{
-									// @todo MeshPaint: We rely on filtering out normal maps by name here.  Obviously a user can name a diffuse with _N_ in the name so
-									//   this is not a good option.  We attempted to find all the normal maps from the material above with GetAllNormalParameterNames(),
-									//   but that always seems to return an empty list.  This needs to be revisited.
+					if( Texture2D != NULL )
+					{
+						// @todo MeshPaint: We rely on filtering out normal maps by name here.  Obviously a user can name a diffuse with _N_ in the name so
+						//   this is not a good option.  We attempted to find all the normal maps from the material above with GetAllNormalParameterNames(),
+						//   but that always seems to return an empty list.  This needs to be revisited.
 	
-									// Some normalmaps in the content will fail checks we do in the if statement below.  So we also check to make sure 
-									//   the name does not end with "_N", and that the following substrings do not appear in the name "_N_" "_N0".
-									FString Texture2DName;
-									Texture2D->GetName(Texture2DName);
-									Texture2DName = Texture2DName.ToUpper();
-
-									// Make sure the texture is not a normalmap, we don't support painting on those at the moment.
-									if( Texture2D->IsNormalMap() == true 
-										|| Texture2D->LODGroup == TEXTUREGROUP_WorldNormalMap
-										|| Texture2D->LODGroup == TEXTUREGROUP_CharacterNormalMap
-										|| Texture2D->LODGroup == TEXTUREGROUP_WeaponNormalMap
-										|| Texture2D->LODGroup == TEXTUREGROUP_VehicleNormalMap
-										|| Texture2D->LODGroup == TEXTUREGROUP_WorldNormalMap
-										|| Texture2DName.Contains( TEXT("_N0" ))
-										|| Texture2DName.Contains( TEXT("_N_" ))
-										|| Texture2DName.Contains( TEXT("_NORMAL" ))
-										|| (Texture2DName.Right(2)).Contains( TEXT("_N" )) )
-									{
-										continue;
-									}
+						// Some normalmaps in the content will fail checks we do in the if statement below.  So we also check to make sure 
+						//   the name does not end with "_N", and that the following substrings do not appear in the name "_N_" "_N0".
+						FString Texture2DName;
+						Texture2D->GetName(Texture2DName);
+						Texture2DName = Texture2DName.ToUpper();
+
+						// Make sure the texture is not a normalmap, we don't support painting on those at the moment.
+						if( Texture2D->IsNormalMap() == true 
+							|| Texture2D->LODGroup == TEXTUREGROUP_WorldNormalMap
+							|| Texture2D->LODGroup == TEXTUREGROUP_CharacterNormalMap
+							|| Texture2D->LODGroup == TEXTUREGROUP_WeaponNormalMap
+							|| Texture2D->LODGroup == TEXTUREGROUP_VehicleNormalMap
+							|| Texture2D->LODGroup == TEXTUREGROUP_WorldNormalMap
+							|| Texture2DName.Contains( TEXT("_N0" ))
+							|| Texture2DName.Contains( TEXT("_N_" ))
+							|| Texture2DName.Contains( TEXT("_NORMAL" ))
+							|| (Texture2DName.Right(2)).Contains( TEXT("_N" )) )
+						{
+							continue;
+						}
 			
-									// Add the texture to our list
-									new(TexturePaintTargetList) FTextureTargetListInfo(Texture2D, UVChannelIndex);
+						// Add the texture to our list
+						new(TexturePaintTargetList) FTextureTargetListInfo(Texture2D, UVChannelIndex);
 
-									// We stored off the user's selection before we began the update.  Since we cleared the list we lost
-									//  that selection info. If the same texture appears in our list after update, we will select it again.
-									if( PreviouslySelectedTexture != NULL && Texture2D == PreviouslySelectedTexture )
-									{
-										TexturePaintTargetList[ TexturePaintTargetList.Num() - 1 ].bIsSelected = true;
-									}
-								}
-							}
+						// We stored off the user's selection before we began the update.  Since we cleared the list we lost
+						//  that selection info. If the same texture appears in our list after update, we will select it again.
+						if( PreviouslySelectedTexture != NULL && Texture2D == PreviouslySelectedTexture )
+						{
+							TexturePaintTargetList[ TexturePaintTargetList.Num() - 1 ].bIsSelected = true;
+						}
+					}
+				}
 
-							//if there are no default textures, revert to the old method of just selecting the first texture.
-							if (DefaultIndex == INDEX_NONE)
-							{
-								DefaultIndex = 0;
-							}
+				//if there are no default textures, revert to the old method of just selecting the first texture.
+				if (DefaultIndex == INDEX_NONE)
+				{
+					DefaultIndex = 0;
+				}
 
-							//We refreshed the list, if nothing else is set we default to the first texture that has IsDefaultMeshPaintTexture set.
-							if(TexturePaintTargetList.Num() > 0 && GetSelectedTexture() == NULL)
-							{
-								if (ensure(TexturePaintTargetList.IsValidIndex(DefaultIndex)))
-								{
-									TexturePaintTargetList[DefaultIndex].bIsSelected = true;
-								}
-							}
-						}
+				//We refreshed the list, if nothing else is set we default to the first texture that has IsDefaultMeshPaintTexture set.
+				if ((TexturePaintTargetList.Num() > 0) && (GetSelectedTexture() == NULL))
+				{
+					if (ensure(TexturePaintTargetList.IsValidIndex(DefaultIndex)))
+					{
+						TexturePaintTargetList[DefaultIndex].bIsSelected = true;
 					}
 				}
-		
+
 			}
 		}
 		
diff --git a/Engine/Source/Editor/MeshPaint/Private/MeshPaintModule.cpp b/Engine/Source/Editor/MeshPaint/Private/MeshPaintModule.cpp
--- a/Engine/Source/Editor/MeshPaint/Private/MeshPaintModule.cpp
+++ b/Engine/Source/Editor/MeshPaint/Private/MeshPaintModule.cpp
@@ -5,6 +5,76 @@
 #include "MeshPaintAdapterFactory.h"
 #include "MeshPaintStaticMeshAdapter.h"
 
+#include "Materials/MaterialExpressionTextureBase.h"
+#include "Materials/MaterialExpressionTextureSample.h"
+#include "Materials/MaterialExpressionTextureCoordinate.h"
+#include "Materials/MaterialExpressionTextureSampleParameter.h"
+
+//////////////////////////////////////////////////////////////////////////
+// IMeshPaintGeometryAdapter
+
+void IMeshPaintGeometryAdapter::DefaultQueryPaintableTextures(int32 MaterialIndex, UMeshComponent* MeshComponent, int32& OutDefaultIndex, TArray<struct FPaintableTexture>& InOutTextureList)
+{
+	OutDefaultIndex = INDEX_NONE;
+
+	// We already know the material we are painting on, take it off the static mesh component
+	UMaterialInterface* Material = MeshComponent->GetMaterial(MaterialIndex);
+
+	if (Material != NULL)
+	{
+		FPaintableTexture PaintableTexture;
+		// Find all the unique textures used in the top material level of the selected actor materials
+
+		const TArray<UMaterialExpression*>& Expressions = Material->GetMaterial()->Expressions;
+
+		// Only grab the textures from the top level of samples
+		for (auto ItExpressions = Expressions.CreateConstIterator(); ItExpressions; ItExpressions++)
+		{
+			UMaterialExpressionTextureBase* TextureBase = Cast<UMaterialExpressionTextureBase>(*ItExpressions);
+			if (TextureBase != NULL &&
+				TextureBase->Texture != NULL &&
+				!TextureBase->Texture->IsNormalMap())
+			{
+				// Default UV channel to index 0. 
+				PaintableTexture = FPaintableTexture(TextureBase->Texture, 0);
+
+				// Texture Samples can have UV's specified, check the first node for whether it has a custom UV channel set. 
+				// We only check the first as the Mesh paint mode does not support painting with UV's modified in the shader.
+				UMaterialExpressionTextureSample* TextureSample = Cast<UMaterialExpressionTextureSample>(*ItExpressions);
+				if (TextureSample != NULL)
+				{
+					UMaterialExpressionTextureCoordinate* TextureCoords = Cast<UMaterialExpressionTextureCoordinate>(TextureSample->Coordinates.Expression);
+					if (TextureCoords != NULL)
+					{
+						// Store the uv channel, this is set when the texture is selected. 
+						PaintableTexture.UVChannelIndex = TextureCoords->CoordinateIndex;
+					}
+
+					// Handle texture parameter expressions
+					UMaterialExpressionTextureSampleParameter* TextureSampleParameter = Cast<UMaterialExpressionTextureSampleParameter>(TextureSample);
+					if (TextureSampleParameter != NULL)
+					{
+						// Grab the overridden texture if it exists.  
+						Material->GetTextureParameterValue(TextureSampleParameter->ParameterName, PaintableTexture.Texture);
+					}
+				}
+
+				// note that the same texture will be added again if its UV channel differs. 
+				int32 TextureIndex = InOutTextureList.AddUnique(PaintableTexture);
+
+				// cache the first default index, if there is no previous info this will be used as the selected texture
+				if ((OutDefaultIndex == INDEX_NONE) && TextureBase->IsDefaultMeshpaintTexture)
+				{
+					OutDefaultIndex = TextureIndex;
+				}
+			}
+		}
+	}
+}
+
+//////////////////////////////////////////////////////////////////////////
+// FMeshPaintModule
+
 class FMeshPaintModule : public IMeshPaintModule
 {
 public:
diff --git a/Engine/Source/Editor/MeshPaint/Private/MeshPaintStaticMeshAdapter.cpp b/Engine/Source/Editor/MeshPaint/Private/MeshPaintStaticMeshAdapter.cpp
--- a/Engine/Source/Editor/MeshPaint/Private/MeshPaintStaticMeshAdapter.cpp
+++ b/Engine/Source/Editor/MeshPaint/Private/MeshPaintStaticMeshAdapter.cpp
@@ -159,6 +159,11 @@ void FMeshPaintGeometryAdapterForStaticMeshes::SphereIntersectTriangles(TArray<i
 	//@TODO: The code for static mesh components is still inside of the mesh painting mode due to the octree
 }
 
+void FMeshPaintGeometryAdapterForStaticMeshes::QueryPaintableTextures(int32 MaterialIndex, int32& OutDefaultIndex, TArray<struct FPaintableTexture>& InOutTextureList)
+{
+	DefaultQueryPaintableTextures(MaterialIndex, StaticMeshComponent, OutDefaultIndex, InOutTextureList);
+}
+
 //////////////////////////////////////////////////////////////////////////
 // FMeshPaintGeometryAdapterForStaticMeshesFactory
 
diff --git a/Engine/Source/Editor/MeshPaint/Private/MeshPaintStaticMeshAdapter.h b/Engine/Source/Editor/MeshPaint/Private/MeshPaintStaticMeshAdapter.h
--- a/Engine/Source/Editor/MeshPaint/Private/MeshPaintStaticMeshAdapter.h
+++ b/Engine/Source/Editor/MeshPaint/Private/MeshPaintStaticMeshAdapter.h
@@ -17,6 +17,7 @@ public:
 	virtual bool SupportsVertexPaint() const override { return true; }
 	virtual bool LineTraceComponent(struct FHitResult& OutHit, const FVector Start, const FVector End, const struct FCollisionQueryParams& Params) const override;
 	virtual void SphereIntersectTriangles(TArray<int32>& OutTriangles, const float ComponentSpaceSquaredBrushRadius, const FVector& ComponentSpaceBrushPosition) const override;
+	virtual void QueryPaintableTextures(int32 MaterialIndex, int32& OutDefaultIndex, TArray<struct FPaintableTexture>& InOutTextureList) override;
 
 protected:
 	UStaticMeshComponent* StaticMeshComponent;
diff --git a/Engine/Source/Editor/MeshPaint/Public/MeshPaintEdMode.h b/Engine/Source/Editor/MeshPaint/Public/MeshPaintEdMode.h
--- a/Engine/Source/Editor/MeshPaint/Public/MeshPaintEdMode.h
+++ b/Engine/Source/Editor/MeshPaint/Public/MeshPaintEdMode.h
@@ -28,6 +28,25 @@ struct FTexturePaintTriangleInfo
 	FVector2D TriUVs[3];
 };
 
+/** Structure used to house and compare Texture and UV channel pairs */
+struct FPaintableTexture
+{	
+	UTexture*	Texture;
+	int32		UVChannelIndex;
+
+	FPaintableTexture(UTexture* InTexture = NULL, uint32 InUVChannelIndex = 0)
+		: Texture(InTexture)
+		, UVChannelIndex(InUVChannelIndex)
+	{}
+	
+	/** Overloaded equality operator for use with TArrays Contains method. */
+	bool operator==(const FPaintableTexture& rhs) const
+	{
+		return (Texture == rhs.Texture);
+		/* && (UVChannelIndex == rhs.UVChannelIndex);*/// if we compare UVChannel we would have to duplicate the texture
+	}
+};
+
 /** Mesh paint mode */
 namespace EMeshPaintMode
 {
diff --git a/Engine/Source/Editor/MeshPaint/Public/MeshPaintModule.h b/Engine/Source/Editor/MeshPaint/Public/MeshPaintModule.h
--- a/Engine/Source/Editor/MeshPaint/Public/MeshPaintModule.h
+++ b/Engine/Source/Editor/MeshPaint/Public/MeshPaintModule.h
@@ -16,7 +16,10 @@ public:
 	virtual bool SupportsVertexPaint() const = 0;
 	virtual bool LineTraceComponent(struct FHitResult& OutHit, const FVector Start, const FVector End, const struct FCollisionQueryParams& Params) const = 0;
 	virtual void SphereIntersectTriangles(TArray<int32>& OutTriangles, const float ComponentSpaceSquaredBrushRadius, const FVector& ComponentSpaceBrushPosition) const = 0;
+	virtual void QueryPaintableTextures(int32 MaterialIndex, int32& OutDefaultIndex, TArray<struct FPaintableTexture>& InOutTextureList) = 0;
 	virtual ~IMeshPaintGeometryAdapter() {}
+
+	MESHPAINT_API static void DefaultQueryPaintableTextures(int32 MaterialIndex, UMeshComponent* MeshComponent, int32& OutDefaultIndex, TArray<struct FPaintableTexture>& InOutTextureList);
 };
 
 /**
