Fixed code relying on SLATE_TEXT_ATTRIBUTE for STextBlock.

diff --git a/Engine/Source/Editor/LevelEditor/Private/SSurfaceProperties.cpp b/Engine/Source/Editor/LevelEditor/Private/SSurfaceProperties.cpp
--- a/Engine/Source/Editor/LevelEditor/Private/SSurfaceProperties.cpp
+++ b/Engine/Source/Editor/LevelEditor/Private/SSurfaceProperties.cpp
@@ -450,7 +450,7 @@ TSharedRef<SWidget> SSurfaceProperties::ConstructTextureScale()
 				.Label()
 				[
 					SNew(STextBlock)
-					.Text(ControlLabels[idx])
+					.Text(FText::FromString(ControlLabels[idx]))
 				]
 			]
 			.MenuContent()
@@ -868,7 +868,7 @@ TSharedRef< ITableRow > SSurfaceProperties::OnGenerateScaleTableRow( TSharedPtr<
 {
 	return SNew(STableRow<TSharedPtr<FString>>, OwnerTable)
 	[
-		SNew(STextBlock) .Text(*Item.Get())
+		SNew(STextBlock) .Text(FText::FromString(*Item.Get()))
 	];
 }
 
diff --git a/Engine/Source/Editor/Levels/Private/LevelViewModel.cpp b/Engine/Source/Editor/Levels/Private/LevelViewModel.cpp
--- a/Engine/Source/Editor/Levels/Private/LevelViewModel.cpp
+++ b/Engine/Source/Editor/Levels/Private/LevelViewModel.cpp
@@ -118,19 +118,18 @@ FString FLevelViewModel::GetName(bool bForceDisplayPath /*=false*/, bool bDispla
 	return DisplayName;
 }
 
-FString FLevelViewModel::GetDisplayName() const
+FText FLevelViewModel::GetDisplayName() const
 {
-	return GetName();
+	return FText::FromString(GetName());
 }
 
-FString FLevelViewModel::GetActorCountString() const
+FText FLevelViewModel::GetActorCountString() const
 {
-	return FString::Printf(TEXT( "%i" ), LevelActorsCount);
+	return FText::AsNumber(LevelActorsCount);
 }
 
-FString FLevelViewModel::GetLightmassSizeString() const
+FText FLevelViewModel::GetLightmassSizeString() const
 {
-	FString MemorySizeString;
 	ULevel* CurLevel = Level.Get();
 
 	if ( CurLevel && GetDefault<ULevelBrowserSettings>()->bDisplayLightmassSize  )
@@ -139,15 +138,17 @@ FString FLevelViewModel::GetLightmassSizeString() const
 		static const float ByteConversion = 1.0f / 1024.0f;
 		float LightmapSize = CurLevel->LightmapTotalSize * ByteConversion;
 		
-		MemorySizeString += FString::Printf( TEXT( "%.2f" ), LightmapSize );
+		static const FNumberFormattingOptions FormatOptions = FNumberFormattingOptions()
+			.SetMinimumFractionalDigits(2)
+			.SetMaximumFractionalDigits(2);
+		return FText::AsNumber(LightmapSize, &FormatOptions);
 	}
 
-	return MemorySizeString;
+	return FText::GetEmpty();
 }
 
-FString FLevelViewModel::GetFileSizeString() const
+FText FLevelViewModel::GetFileSizeString() const
 {
-	FString MemorySizeString;
 	ULevel* CurLevel = Level.Get();
 
 	if ( CurLevel && GetDefault<ULevelBrowserSettings>()->bDisplayFileSize  )
@@ -156,10 +157,13 @@ FString FLevelViewModel::GetFileSizeString() const
 		static const float ByteConversion = 1.0f / 1024.0f;
 		float FileSize = CurLevel->GetOutermost()->GetFileSize() * ByteConversion * ByteConversion;
 		
-		MemorySizeString += FString::Printf( TEXT( "%.2f" ), FileSize );
+		static const FNumberFormattingOptions FormatOptions = FNumberFormattingOptions()
+			.SetMinimumFractionalDigits(2)
+			.SetMaximumFractionalDigits(2);
+		return FText::AsNumber(FileSize, &FormatOptions);
 	}
 
-	return MemorySizeString;
+	return FText::GetEmpty();
 }
 
 void FLevelViewModel::ClearDirtyFlag()
diff --git a/Engine/Source/Editor/Levels/Private/LevelViewModel.h b/Engine/Source/Editor/Levels/Private/LevelViewModel.h
--- a/Engine/Source/Editor/Levels/Private/LevelViewModel.h
+++ b/Engine/Source/Editor/Levels/Private/LevelViewModel.h
@@ -75,18 +75,18 @@ public:
 	  * @param	bDisplayTags		Whether to include tags (ie: Persistent/Current/Dirty) */
 	FString GetName(bool bForceDisplayPath=false, bool bDisplayTags=true) const;
 
-	/**	@return	The Level's display name as a FString;
+	/**	@return	The Level's display name as a FText;
 	  * wrapper function for GetName that simplifies function pointer for slate */
-	FString GetDisplayName() const;
+	FText GetDisplayName() const;
 
-	/** @return The Level's Actor Count as a FString */
-	FString GetActorCountString() const;
+	/** @return The Level's Actor Count as a FText */
+	FText GetActorCountString() const;
 
-	/** @return the Level's Lightmass Size as a FString */
-	FString GetLightmassSizeString() const;
+	/** @return the Level's Lightmass Size as a FText */
+	FText GetLightmassSizeString() const;
 
-	/** @return the Level's File Size as a FString */
-	FString GetFileSizeString() const;
+	/** @return the Level's File Size as a FText */
+	FText GetFileSizeString() const;
 
 	/** @return whether the Level is dirty */
 	bool IsDirty() const;
diff --git a/Engine/Source/Editor/MainFrame/Private/MainFrameModule.cpp b/Engine/Source/Editor/MainFrame/Private/MainFrameModule.cpp
--- a/Engine/Source/Editor/MainFrame/Private/MainFrameModule.cpp
+++ b/Engine/Source/Editor/MainFrame/Private/MainFrameModule.cpp
@@ -224,25 +224,32 @@ TSharedRef<SWidget> FMainFrameModule::MakeDeveloperTools() const
 {
 	struct Local
 	{
-		static FString GetFrameRateAsString() 
+		static FText GetFrameRateAsString() 
 		{
 			// Clamp to avoid huge averages at startup or after hitches
 			const float AverageFPS = 1.0f / FSlateApplication::Get().GetAverageDeltaTime();
 			const float ClampedFPS = ( AverageFPS < 0.0f || AverageFPS > 4000.0f ) ? 0.0f : AverageFPS;
 
-			return FString::Printf( TEXT( "% 3.1f" ), ClampedFPS );
+			static const FNumberFormattingOptions FormatOptions = FNumberFormattingOptions()
+				.SetMinimumFractionalDigits(1)
+				.SetMaximumFractionalDigits(1);
+			return FText::AsNumber( ClampedFPS, &FormatOptions );
 		}
 
-		static FString GetFrameTimeAsString() 
+		static FText GetFrameTimeAsString() 
 		{
 			// Clamp to avoid huge averages at startup or after hitches
 			const float AverageMS = FSlateApplication::Get().GetAverageDeltaTime() * 1000.0f;
 			const float ClampedMS = ( AverageMS < 0.0f || AverageMS > 4000.0f ) ? 0.0f : AverageMS;
 
-			return FString::Printf( TEXT( "% 3.1f ms" ), ClampedMS );
+			static const FNumberFormattingOptions FormatOptions = FNumberFormattingOptions()
+				.SetMinimumFractionalDigits(1)
+				.SetMaximumFractionalDigits(1);
+			static const FText FrameTimeFmt = FText::FromString(TEXT("{0} ms"));
+			return FText::Format( FrameTimeFmt, FText::AsNumber( ClampedMS, &FormatOptions ) );
 		}
 
-		static FString GetMemoryAsString() 
+		static FText GetMemoryAsString() 
 		{
 			// Only refresh process memory allocated after every so often, to reduce fixed frame time overhead
 			static SIZE_T StaticLastTotalAllocated = 0;
@@ -257,12 +264,16 @@ TSharedRef<SWidget> FMainFrameModule::MakeDeveloperTools() const
 				QueriesUntilUpdate = 60;
 			}
 
-			return FString::Printf( TEXT( "% 5.2f mb" ), (float)StaticLastTotalAllocated / ( 1024.0f * 1024.0f ) );
+			static const FNumberFormattingOptions FormatOptions = FNumberFormattingOptions()
+				.SetMinimumFractionalDigits(2)
+				.SetMaximumFractionalDigits(2);
+			static const FText MemorySizeFmt = FText::FromString(TEXT("{0} mb"));
+			return FText::Format( MemorySizeFmt, FText::AsNumber( (float)StaticLastTotalAllocated / ( 1024.0f * 1024.0f ), &FormatOptions ) );
 		}
 
-		static FString GetUObjectCountAsString() 
+		static FText GetUObjectCountAsString() 
 		{
-			return FString::Printf( TEXT( " %i"), GUObjectArray.GetObjectArrayNumMinusAvailable() );
+			return FText::AsNumber(GUObjectArray.GetObjectArrayNumMinusAvailable());
 		}
 
 		static void OpenVideo( FString SourceFilePath )
@@ -379,7 +390,7 @@ TSharedRef<SWidget> FMainFrameModule::MakeDeveloperTools() const
 			.VAlign(VAlign_Bottom)
 			[
 				SNew( STextBlock )
-				.Text( LOCTEXT("FrameRateLabel", "FPS:") )
+				.Text( LOCTEXT("FrameRateLabel", "FPS: ") )
 				.Font( LabelFont )
 				.ColorAndOpacity( FLinearColor( 0.3f, 0.3f, 0.3f ) )
 			]
@@ -426,7 +437,7 @@ TSharedRef<SWidget> FMainFrameModule::MakeDeveloperTools() const
 				.VAlign(VAlign_Bottom)
 				[
 					SNew( STextBlock )
-					.Text( LOCTEXT("MemoryLabel", "Mem:") )
+					.Text( LOCTEXT("MemoryLabel", "Mem: ") )
 					.Font( LabelFont )
 					.ColorAndOpacity( FLinearColor( 0.3f, 0.3f, 0.3f ) )
 				]
@@ -453,7 +464,7 @@ TSharedRef<SWidget> FMainFrameModule::MakeDeveloperTools() const
 				.VAlign(VAlign_Bottom)
 				[
 					SNew( STextBlock )
-						.Text( LOCTEXT("UObjectCountLabel", "Objs:") )
+						.Text( LOCTEXT("UObjectCountLabel", "Objs: ") )
 						.Font( LabelFont )
 						.ColorAndOpacity( FLinearColor( 0.3f, 0.3f, 0.3f ) )
 				]
diff --git a/Engine/Source/Editor/MaterialEditor/Private/FindInMaterial.cpp b/Engine/Source/Editor/MaterialEditor/Private/FindInMaterial.cpp
--- a/Engine/Source/Editor/MaterialEditor/Private/FindInMaterial.cpp
+++ b/Engine/Source/Editor/MaterialEditor/Private/FindInMaterial.cpp
@@ -47,17 +47,17 @@ FReply FFindInMaterialResult::OnClick(TWeakPtr<class FMaterialEditor> MaterialEd
 	return FReply::Handled();
 }
 
-FString FFindInMaterialResult::GetCategory() const
+FText FFindInMaterialResult::GetCategory() const
 {
 	if (Class == UEdGraphPin::StaticClass())
 	{
-		return LOCTEXT("PinCategory", "Pin").ToString();
+		return LOCTEXT("PinCategory", "Pin");
 	}
 	else
 	{
-		return LOCTEXT("NodeCategory", "Node").ToString();
+		return LOCTEXT("NodeCategory", "Node");
 	}
-	return FString();
+	return FText::GetEmpty();
 }
 
 TSharedRef<SWidget> FFindInMaterialResult::CreateIcon() const
@@ -314,12 +314,12 @@ void SFindInMaterial::MatchTokens(const TArray<FString> &Tokens)
 			UEdGraphPin* Pin = *PinIt;
 			if (Pin && Pin->PinFriendlyName.CompareTo(FText::FromString(TEXT(" "))) != 0)
 			{
-				FString PinName = Pin->GetSchema()->GetPinDisplayName(Pin);
+				FText PinName = Pin->GetSchema()->GetPinDisplayName(Pin);
 				FString PinSearchString = Pin->PinName + Pin->PinFriendlyName.ToString() + Pin->DefaultValue + Pin->PinType.PinCategory + Pin->PinType.PinSubCategory + (Pin->PinType.PinSubCategoryObject.IsValid() ? Pin->PinType.PinSubCategoryObject.Get()->GetFullName() : TEXT(""));
 				PinSearchString = PinSearchString.Replace(TEXT(" "), TEXT(""));
 				if (StringMatchesSearchTokens(Tokens, PinSearchString))
 				{
-					FSearchResult PinResult(new FFindInMaterialResult(PinName, NodeResult, Pin));
+					FSearchResult PinResult(new FFindInMaterialResult(PinName.ToString(), NodeResult, Pin));
 					NodeResult->Children.Add(PinResult);
 				}
 			}
@@ -347,8 +347,8 @@ TSharedRef<ITableRow> SFindInMaterial::OnGenerateRow(FSearchResult InItem, const
 				.ForegroundColor(FEditorStyle::GetColor("PropertyWindow.CategoryForeground"))
 				[
 					SNew(STextBlock)
-					.Text(InItem->Value)
-					.ToolTipText(LOCTEXT("BlueprintCatSearchToolTip", "Blueprint").ToString())
+					.Text(FText::FromString(InItem->Value))
+					.ToolTipText(LOCTEXT("BlueprintCatSearchToolTip", "Blueprint"))
 				]
 			];
 	}
@@ -370,9 +370,9 @@ TSharedRef<ITableRow> SFindInMaterial::OnGenerateRow(FSearchResult InItem, const
 					.Padding(2, 0)
 					[
 						SNew(STextBlock)
-						.Text(InItem->Value)
+						.Text(FText::FromString(InItem->Value))
 						.HighlightText(HighlightText)
-						.ToolTipText(FString::Printf(*LOCTEXT("BlueprintResultSearchToolTip", "%s : %s").ToString(), *InItem->GetCategory(), *InItem->Value))
+						.ToolTipText(FText::Format(LOCTEXT("BlueprintResultSearchToolTipFmt", "{0} : {1}"), InItem->GetCategory(), FText::FromString(InItem->Value)))
 					]
 				+ SHorizontalBox::Slot()
 					.FillWidth(1)
@@ -381,7 +381,7 @@ TSharedRef<ITableRow> SFindInMaterial::OnGenerateRow(FSearchResult InItem, const
 					.Padding(2, 0)
 					[
 						SNew(STextBlock)
-						.Text(InItem->GetValueText())
+						.Text(FText::FromString(InItem->GetValueText()))
 						.HighlightText(HighlightText)
 					]
 				+ SHorizontalBox::Slot()
@@ -391,7 +391,7 @@ TSharedRef<ITableRow> SFindInMaterial::OnGenerateRow(FSearchResult InItem, const
 					.Padding(2, 0)
 					[
 						SNew(STextBlock)
-						.Text(CommentText.IsEmpty() ? FString() : FString::Printf(TEXT("Node Comment:[%s]"), *CommentText))
+						.Text(CommentText.IsEmpty() ? FText::GetEmpty() : FText::Format(LOCTEXT("NodeCommentFmt", "Node Comment:[{0}]"), FText::FromString(CommentText)))
 						.ColorAndOpacity(FLinearColor::Yellow)
 						.HighlightText(HighlightText)
 					]
diff --git a/Engine/Source/Editor/MaterialEditor/Private/MaterialEditor.cpp b/Engine/Source/Editor/MaterialEditor/Private/MaterialEditor.cpp
--- a/Engine/Source/Editor/MaterialEditor/Private/MaterialEditor.cpp
+++ b/Engine/Source/Editor/MaterialEditor/Private/MaterialEditor.cpp
@@ -1104,15 +1104,14 @@ void FMaterialEditor::SaveEditorSettings()
 	GConfig->SetInt(TEXT("MaterialEditor"), TEXT("PrimType"), Viewport->PreviewPrimType, GEditorUserSettingsIni);
 }
 
-FString FMaterialEditor::GetCodeViewText() const
+FText FMaterialEditor::GetCodeViewText() const
 {
-	return HLSLCode;
+	return FText::FromString(HLSLCode);
 }
 
 FReply FMaterialEditor::CopyCodeViewTextToClipboard()
 {
-	FString CodeViewText = GetCodeViewText();
-	FPlatformMisc::ClipboardCopy( *CodeViewText );
+	FPlatformMisc::ClipboardCopy(*HLSLCode);
 	return FReply::Handled();
 }
 
@@ -3765,11 +3764,11 @@ TSharedRef<SGraphEditor> FMaterialEditor::CreateGraphEditorWidget()
 	
 	if (MaterialFunction)
 	{
-		AppearanceInfo.CornerText = LOCTEXT("AppearanceCornerText_MaterialFunction", "MATERIAL FUNCTION").ToString();
+		AppearanceInfo.CornerText = LOCTEXT("AppearanceCornerText_MaterialFunction", "MATERIAL FUNCTION");
 	}
 	else
 	{
-		AppearanceInfo.CornerText = LOCTEXT("AppearanceCornerText_Material", "MATERIAL").ToString();
+		AppearanceInfo.CornerText = LOCTEXT("AppearanceCornerText_Material", "MATERIAL");
 	}
 
 	SGraphEditor::FGraphEditorEvents InEvents;
diff --git a/Engine/Source/Editor/MaterialEditor/Private/MaterialEditor.h b/Engine/Source/Editor/MaterialEditor/Private/MaterialEditor.h
--- a/Engine/Source/Editor/MaterialEditor/Private/MaterialEditor.h
+++ b/Engine/Source/Editor/MaterialEditor/Private/MaterialEditor.h
@@ -459,7 +459,7 @@ private:
 	void SaveEditorSettings();
 
 	/** Gets the text in the code view widget */
-	FString GetCodeViewText() const;
+	FText GetCodeViewText() const;
 
 	/** Copies all the HLSL Code View code to the clipboard */
 	FReply CopyCodeViewTextToClipboard();
diff --git a/Engine/Source/Editor/MaterialEditor/Private/MaterialEditorDetailCustomization.cpp b/Engine/Source/Editor/MaterialEditor/Private/MaterialEditorDetailCustomization.cpp
--- a/Engine/Source/Editor/MaterialEditor/Private/MaterialEditorDetailCustomization.cpp
+++ b/Engine/Source/Editor/MaterialEditor/Private/MaterialEditorDetailCustomization.cpp
@@ -92,7 +92,7 @@ TSharedRef< ITableRow > FMaterialExpressionParameterDetails::MakeDetailsGroupVie
 {
 	return SNew(STableRow<TSharedPtr<FString>>, OwnerTable)
 		[
-			SNew(STextBlock) .Text(*Item.Get())
+			SNew(STextBlock) .Text(FText::FromString(*Item.Get()))
 		];
 }
 
@@ -151,17 +151,17 @@ FText FMaterialExpressionCollectionParameterDetails::GetToolTipText() const
 	}
 }
 
-FString FMaterialExpressionCollectionParameterDetails::GetParameterNameString() const
+FText FMaterialExpressionCollectionParameterDetails::GetParameterNameString() const
 {
 	FString CurrentParameterName;
 
 	FPropertyAccess::Result Result = ParameterNamePropertyHandle->GetValue(CurrentParameterName);
 	if( Result == FPropertyAccess::MultipleValues )
 	{
-		CurrentParameterName = NSLOCTEXT("PropertyEditor", "MultipleValues", "Multiple Values").ToString();
+		return NSLOCTEXT("PropertyEditor", "MultipleValues", "Multiple Values");
 	}
 
-	return CurrentParameterName;
+	return FText::FromString(CurrentParameterName);
 }
 
 bool FMaterialExpressionCollectionParameterDetails::IsParameterNameComboEnabled() const
@@ -285,7 +285,7 @@ TSharedRef< ITableRow > FMaterialExpressionCollectionParameterDetails::MakeDetai
 {
 	return SNew(STableRow<TSharedPtr<FString>>, OwnerTable)
 		[
-			SNew(STextBlock) .Text(*Item.Get())
+			SNew(STextBlock) .Text(FText::FromString(*Item.Get()))
 		];
 }
 
diff --git a/Engine/Source/Editor/MaterialEditor/Private/MaterialEditorDetailCustomization.h b/Engine/Source/Editor/MaterialEditor/Private/MaterialEditorDetailCustomization.h
--- a/Engine/Source/Editor/MaterialEditor/Private/MaterialEditorDetailCustomization.h
+++ b/Engine/Source/Editor/MaterialEditor/Private/MaterialEditorDetailCustomization.h
@@ -71,7 +71,7 @@ private:
 	void OnSelectionChanged( TSharedPtr<FString> ProposedSelection, ESelectInfo::Type SelectInfo );
 
 	FText GetToolTipText() const;
-	FString GetParameterNameString() const;
+	FText GetParameterNameString() const;
 	bool IsParameterNameComboEnabled() const;
 	void OnCollectionChanged();
 
diff --git a/Engine/Source/Editor/MaterialEditor/Private/SMaterialEditorTitleBar.cpp b/Engine/Source/Editor/MaterialEditor/Private/SMaterialEditorTitleBar.cpp
--- a/Engine/Source/Editor/MaterialEditor/Private/SMaterialEditorTitleBar.cpp
+++ b/Engine/Source/Editor/MaterialEditor/Private/SMaterialEditorTitleBar.cpp
@@ -70,7 +70,7 @@ TSharedRef<ITableRow> SMaterialEditorTitleBar::MakeMaterialInfoWidget(TSharedPtr
 				SNew(STextBlock)
 				.ColorAndOpacity(TextColor)
 				.Font(FSlateFontInfo( FPaths::EngineContentDir() / TEXT("Slate/Fonts/Roboto-Regular.ttf"), FontSize ))
-				.Text( Text )
+				.Text(FText::FromString(Text))
 			];
 	}
 }
diff --git a/Engine/Source/Editor/MaterialEditor/Private/SMaterialPalette.cpp b/Engine/Source/Editor/MaterialEditor/Private/SMaterialPalette.cpp
--- a/Engine/Source/Editor/MaterialEditor/Private/SMaterialPalette.cpp
+++ b/Engine/Source/Editor/MaterialEditor/Private/SMaterialPalette.cpp
@@ -39,7 +39,7 @@ void SMaterialPaletteItem::Construct(const FArguments& InArgs, FCreateWidgetForA
 	const FSlateBrush* IconBrush = FEditorStyle::GetBrush(TEXT("NoBrush"));
 	FSlateColor IconColor = FSlateColor::UseForeground();
 	FString ToolTip = GraphAction->TooltipDescription;
-	FString IconToolTip = ToolTip;
+	FText IconToolTip = FText::FromString(ToolTip);
 	bool bIsReadOnly = false;
 
 	TSharedRef<SWidget> IconWidget = CreateIconWidget( IconToolTip, IconBrush, IconColor );
diff --git a/Engine/Source/Editor/MaterialEditor/Public/FindInMaterial.h b/Engine/Source/Editor/MaterialEditor/Public/FindInMaterial.h
--- a/Engine/Source/Editor/MaterialEditor/Public/FindInMaterial.h
+++ b/Engine/Source/Editor/MaterialEditor/Public/FindInMaterial.h
@@ -22,7 +22,7 @@ public:
 	FReply OnClick(TWeakPtr<class FMaterialEditor> MaterialEditor);
 
 	/* Get Category for this search result */
-	FString GetCategory() const;
+	FText GetCategory() const;
 
 	/** Create an icon to represent the result */
 	TSharedRef<SWidget>	CreateIcon() const;
diff --git a/Engine/Source/Editor/MeshPaint/Private/SMeshPaint.cpp b/Engine/Source/Editor/MeshPaint/Private/SMeshPaint.cpp
--- a/Engine/Source/Editor/MeshPaint/Private/SMeshPaint.cpp
+++ b/Engine/Source/Editor/MeshPaint/Private/SMeshPaint.cpp
@@ -301,9 +301,9 @@ private:
 		UVValue = InCount;
 	}
 
-	FString GetUVSelectionString() const
+	FText GetUVSelectionString() const
 	{
-		return FString::Printf(TEXT("%d"), UVValue);
+		return FText::AsNumber(UVValue);
 	}
 
 	void OnChangeLOD(int32 InCount)
@@ -311,9 +311,9 @@ private:
 		LODValue = InCount;
 	}
 
-	FString GetLODSelectionString() const
+	FText GetLODSelectionString() const
 	{
-		return FString::Printf(TEXT("%d"), LODValue);
+		return FText::AsNumber(LODValue);
 	}
 
 	TSharedRef<SWidget> GetUVMenu()
@@ -963,7 +963,7 @@ public:
 					+SVerticalBox::Slot()
 					.AutoHeight()
 					[
-						CreateRadioButton( TEXT("1"), 0 )
+						CreateRadioButton( FText::FromString(TEXT("1")), 0 )
 					]
 				]
 				+SHorizontalBox::Slot()
@@ -975,7 +975,7 @@ public:
 					+SVerticalBox::Slot()
 					.AutoHeight()
 					[
-						CreateRadioButton( TEXT("2"), 1 )
+						CreateRadioButton( FText::FromString(TEXT("2")), 1 )
 					]
 					
 				]
@@ -988,7 +988,7 @@ public:
 					+SVerticalBox::Slot()
 					.AutoHeight()
 					[
-						CreateRadioButton( TEXT("3"), 2 )
+						CreateRadioButton( FText::FromString(TEXT("3")), 2 )
 					]
 				]
 				+SHorizontalBox::Slot()
@@ -1000,7 +1000,7 @@ public:
 					+SVerticalBox::Slot()
 					.AutoHeight()
 					[
-						CreateRadioButton( TEXT("4"), 3 )
+						CreateRadioButton( FText::FromString(TEXT("4")), 3 )
 					]
 					
 				]
@@ -1013,7 +1013,7 @@ public:
 					+SVerticalBox::Slot()
 					.AutoHeight()
 					[
-						CreateRadioButton( TEXT("5"), 4 )
+						CreateRadioButton( FText::FromString(TEXT("5")), 4 )
 					]
 				]
 			]
@@ -1022,7 +1022,7 @@ public:
 
 private:
 
-	TSharedRef<SWidget> CreateRadioButton( const FString& RadioText, int32 RadioButtonChoice )
+	TSharedRef<SWidget> CreateRadioButton( const FText& RadioText, int32 RadioButtonChoice )
 	{
 		return
 			SNew(SCheckBox)
@@ -2164,9 +2164,9 @@ int32 SMeshPaint::GetEraseWeightIndex() const
 	return FMeshPaintSettings::Get().EraseWeightIndex;
 }
 
-FString SMeshPaint::GetInstanceVertexColorsText() const 
+FText SMeshPaint::GetInstanceVertexColorsText() const 
 {
-	FString Text = LOCTEXT("MeshPaint_InstVertexColorsStartText", "None").ToString();
+	FText Text = LOCTEXT("MeshPaint_InstVertexColorsStartText", "None");
 	int32 NumBaseVertexColorBytes = 0;
 	int32 NumInstanceVertexColorBytes = 0;
 	bool bHasInstanceMaterialAndTexture = false;
@@ -2175,7 +2175,10 @@ FString SMeshPaint::GetInstanceVertexColorsText() const
 	if( NumInstanceVertexColorBytes > 0 )
 	{
 		float VertexKiloBytes = NumInstanceVertexColorBytes / 1000.0f;
-		Text = FString::Printf( TEXT( "%.3f k" ), VertexKiloBytes );
+		static const FNumberFormattingOptions FormatOptions = FNumberFormattingOptions()
+			.SetMinimumFractionalDigits(3)
+			.SetMaximumFractionalDigits(3);
+		Text = FText::Format( LOCTEXT("InstanceVertexColorsKiloBytesFmt", "{0} k"), FText::AsNumber(VertexKiloBytes, &FormatOptions) );
 	}
 	return Text;
 }
@@ -2683,23 +2686,23 @@ TSharedRef<SWidget> SMeshPaint::GetTextureTargetWidget( UTexture2D* TextureData
 		];
 }
 
-FString SMeshPaint::GetCurrentTextureTargetText( UTexture2D* TextureData, int index ) const
+FText SMeshPaint::GetCurrentTextureTargetText( UTexture2D* TextureData, int index ) const
 {
 	UTexture2D* Tex = TextureData != NULL ? TextureData :  MeshPaintEditMode->GetSelectedTexture();
 
 	if( Tex == NULL ) 
 	{
-		return TEXT("");
+		return FText::GetEmpty();
 	}
 
 	switch ( index )
 	{
 		case 0:	
-			return Tex->GetName();
+			return FText::FromString(Tex->GetName());
 		case 1:	
-			return Tex->GetDesc();
+			return FText::FromString(Tex->GetDesc());
 		default:
-			return TEXT("");
+			return FText::GetEmpty();
 	}
 }
 
@@ -2709,10 +2712,10 @@ const FSlateBrush* SMeshPaint::GetCurrentTextureTargetImage(UTexture2D* TextureD
 	return Tex != NULL ? new FSlateDynamicImageBrush( Tex, FVector2D(64, 64), Tex->GetFName() ) : NULL;
 }
 
-FString SMeshPaint::GetCurrentUVChannel() const
+FText SMeshPaint::GetCurrentUVChannel() const
 {
-	UTexture2D* Tex = MeshPaintEditMode->GetSelectedTexture();	
-	return FString::Printf( TEXT("%d"), FMeshPaintSettings::Get().UVChannel );	
+	UTexture2D* Tex = MeshPaintEditMode->GetSelectedTexture();
+	return FText::AsNumber( FMeshPaintSettings::Get().UVChannel );	
 }
 
 FReply SMeshPaint::OnChangeTextureTarget( TWeakObjectPtr<UTexture2D> TextureData )
@@ -2825,10 +2828,10 @@ TSharedRef<SWidget> SMeshPaint::GetMaterialSelectionMenu()
 	return MenuBuilder.MakeWidget();
 }
 
-FString SMeshPaint::GetEditingMaterial() const
+FText SMeshPaint::GetEditingMaterial() const
 {
 	int32 SelectedMaterialIndex = MeshPaintEditMode->GetEditingMaterialIndex();
-	return TTypeToString<int32>::ToString(SelectedMaterialIndex);
+	return FText::AsNumber(SelectedMaterialIndex);
 }
 
 void SMeshPaint::OnSetEditingMaterial( int32 NewMaterialIndex )
@@ -2850,14 +2853,14 @@ TSharedRef<SWidget> SMeshPaint::GetActorSelectionMenu()
 	return MenuBuilder.MakeWidget();
 }
 
-FString SMeshPaint::GetEditingActorLabel() const
+FText SMeshPaint::GetEditingActorLabel() const
 {
 	TWeakObjectPtr<AActor> SelectedActor = MeshPaintEditMode->GetEditingActor();
 	if (SelectedActor.IsValid())
 	{
-		return SelectedActor.Get()->GetActorLabel();
+		return FText::FromString(SelectedActor.Get()->GetActorLabel());
 	}
-	return FString();
+	return FText::GetEmpty();
 }
 
 void SMeshPaint::OnSetEditingActor( TWeakObjectPtr<AActor> InActor )
diff --git a/Engine/Source/Editor/MeshPaint/Private/SMeshPaint.h b/Engine/Source/Editor/MeshPaint/Private/SMeshPaint.h
--- a/Engine/Source/Editor/MeshPaint/Private/SMeshPaint.h
+++ b/Engine/Source/Editor/MeshPaint/Private/SMeshPaint.h
@@ -116,7 +116,7 @@ protected:
 	int32 GetEraseWeightIndex() const;
 
 	/** Returns a string representing the memory used selected actors vertex color data */
-	FString GetInstanceVertexColorsText() const;
+	FText GetInstanceVertexColorsText() const;
 
 	/** Called when the Mesh Paint Mode radio button selection is changed */
 	void OntMeshPaintResourceChanged(EMeshPaintResource::Type InPaintResource);
@@ -262,7 +262,7 @@ protected:
 	/** returns a UV channel menu widget */
 	TSharedRef<SWidget> GetUVChannels();
 
-	FString GetCurrentUVChannel() const;
+	FText GetCurrentUVChannel() const;
 
 	/** returns a list of texture target menu widgets */
 	TSharedRef<SWidget> GetTextureTargets();
@@ -271,7 +271,7 @@ protected:
 	TSharedRef<SWidget> GetTextureTargetWidget( UTexture2D* TextureData );
 
 	/** returns a string containing currently selected texture */
-	FString GetCurrentTextureTargetText( UTexture2D* TextureData, int index ) const;
+	FText GetCurrentTextureTargetText( UTexture2D* TextureData, int index ) const;
 
 	/** returns a image containing currently selected texture image */
 	const FSlateBrush* GetCurrentTextureTargetImage( UTexture2D* TextureData ) const;
@@ -301,7 +301,7 @@ protected:
 	TSharedRef<SWidget> GetActorSelectionMenu();
 
 	/** Returns the currently editing actors label */
-	FString GetEditingActorLabel() const;
+	FText GetEditingActorLabel() const;
 
 	/** Called upon the user changing which mesh they wish to act upon */
 	void OnSetEditingActor( TWeakObjectPtr<AActor> InActor );
@@ -313,7 +313,7 @@ protected:
 	TSharedRef<SWidget> GetMaterialSelectionMenu();
 
 	/** Returns the currently editing materials index */
-	FString GetEditingMaterial() const;
+	FText GetEditingMaterial() const;
 
 	/** Called upon the user changing which material they wish to edit */
 	void OnSetEditingMaterial(int32 NewMaterialIndex);
diff --git a/Engine/Source/Editor/NewLevelDialog/Private/NewLevelDialogModule.cpp b/Engine/Source/Editor/NewLevelDialog/Private/NewLevelDialogModule.cpp
--- a/Engine/Source/Editor/NewLevelDialog/Private/NewLevelDialogModule.cpp
+++ b/Engine/Source/Editor/NewLevelDialog/Private/NewLevelDialogModule.cpp
@@ -167,27 +167,26 @@ private:
 	TSharedRef<SWidget> GetWidgetForTemplate(TSharedPtr<FTemplateListItem> Template)
 	{
 		TSharedPtr<SWidget> Image;
-		FString Text;
+		FText Text = FText::GetEmpty();
 		if (Template->bIsNewLevelItem)
 		{
 			// New level item
 			Image = SNew(SImage).Image(FEditorStyle::GetBrush(TEXT("NewLevelDialog.Blank")));
-			Text = LOCTEXT("NewLevelItemLabel", "Empty Level").ToString();
+			Text = LOCTEXT("NewLevelItemLabel", "Empty Level");
 		}
 		else if (Template->TemplateMapInfo.ThumbnailTexture)
 		{
 			// Level with thumbnail
 			Image = SNew(STexture2DView, Template->TemplateMapInfo.ThumbnailTexture);
-			Text = Template->TemplateMapInfo.ThumbnailTexture->GetName();
+			Text = FText::FromString(Template->TemplateMapInfo.ThumbnailTexture->GetName().Replace(TEXT("_"), TEXT(" ")));
 		}
 		else
 		{
 			// Level with no thumbnail
 			Image = SNew(SImage).Image(FEditorStyle::GetBrush(TEXT("NewLevelDialog.Default")));
-			Text = FPackageName::GetShortName(Template->TemplateMapInfo.Map);
+			Text = FText::FromString(FPackageName::GetShortName(Template->TemplateMapInfo.Map).Replace(TEXT("_"), TEXT(" ")));
 		}
 
-		Text.ReplaceInline(TEXT("_"), TEXT(" "));
 		Image->SetCursor(EMouseCursor::Hand);
 
 		return SNew(SBox)
diff --git a/Engine/Source/Editor/NiagaraEditor/Private/NiagaraEditor.cpp b/Engine/Source/Editor/NiagaraEditor/Private/NiagaraEditor.cpp
--- a/Engine/Source/Editor/NiagaraEditor/Private/NiagaraEditor.cpp
+++ b/Engine/Source/Editor/NiagaraEditor/Private/NiagaraEditor.cpp
@@ -147,7 +147,7 @@ TSharedRef<SGraphEditor> FNiagaraEditor::CreateGraphEditorWidget(UEdGraph* InGra
 
 	// Create the appearance info
 	FGraphAppearanceInfo AppearanceInfo;
-	AppearanceInfo.CornerText = LOCTEXT("AppearanceCornerText", "NIAGARA").ToString();
+	AppearanceInfo.CornerText = LOCTEXT("AppearanceCornerText", "NIAGARA");
 
 	SGraphEditor::FGraphEditorEvents InEvents;
 	InEvents.OnSelectionChanged = SGraphEditor::FOnSelectionChanged::CreateSP(this, &FNiagaraEditor::OnSelectedNodesChanged);
diff --git a/Engine/Source/Editor/NiagaraEditor/Private/NiagaraEffectEditor.cpp b/Engine/Source/Editor/NiagaraEditor/Private/NiagaraEffectEditor.cpp
--- a/Engine/Source/Editor/NiagaraEditor/Private/NiagaraEffectEditor.cpp
+++ b/Engine/Source/Editor/NiagaraEditor/Private/NiagaraEffectEditor.cpp
@@ -106,7 +106,7 @@ TSharedRef<SNiagaraEffectEditorWidget> FNiagaraEffectEditor::CreateEditorWidget(
 
 	// Create the appearance info
 	FGraphAppearanceInfo AppearanceInfo;
-	AppearanceInfo.CornerText = LOCTEXT("AppearanceCornerText", "NIAGARA").ToString();
+	AppearanceInfo.CornerText = LOCTEXT("AppearanceCornerText", "NIAGARA");
 
 	// Make title bar
 	TSharedRef<SWidget> TitleBarWidget =
diff --git a/Engine/Source/Editor/NiagaraEditor/Private/SNiagaraEffectEditorWidget.cpp b/Engine/Source/Editor/NiagaraEditor/Private/SNiagaraEffectEditorWidget.cpp
--- a/Engine/Source/Editor/NiagaraEditor/Private/SNiagaraEffectEditorWidget.cpp
+++ b/Engine/Source/Editor/NiagaraEditor/Private/SNiagaraEffectEditorWidget.cpp
@@ -117,7 +117,7 @@ void SNiagaraEffectEditorWidget::Construct(const FArguments& InArgs)
 							SNew(STextBlock)
 							.Visibility(EVisibility::HitTestInvisible)
 							.TextStyle(FEditorStyle::Get(), "Graph.CornerText")
-							.Text(FString("Niagara Effect"))
+							.Text(LOCTEXT("NiagaraEditorLabel", "Niagara Effect"))
 						]
 
 					// Top-right corner text indicating read only when not simulating
@@ -129,7 +129,7 @@ void SNiagaraEffectEditorWidget::Construct(const FArguments& InArgs)
 							SNew(STextBlock)
 							.Visibility(this, &SNiagaraEffectEditorWidget::ReadOnlyVisibility)
 							.TextStyle(FEditorStyle::Get(), "Graph.CornerText")
-							.Text(FString("Read Only"))
+							.Text(LOCTEXT("ReadOnlyLabel", "Read Only"))
 						]
 
 					// Bottom-right corner text for notification list position
@@ -273,7 +273,7 @@ void SEmitterWidget::Construct(const FArguments& InArgs)
 											.Padding(2)
 											[
 												SNew(STextBlock)
-												.Text(FString("Update Script"))
+												.Text(LOCTEXT("UpdateScriptLabel", "Update Script"))
 											]
 											+ SHorizontalBox::Slot()
 											.HAlign(HAlign_Right)
@@ -295,7 +295,7 @@ void SEmitterWidget::Construct(const FArguments& InArgs)
 									.AutoHeight()
 									.Expose(UpdateScriptConstantListSlot)
 									[
-										SNew(SExpandableArea).InitiallyCollapsed(true).AreaTitle(FString("Constants"))
+										SNew(SExpandableArea).InitiallyCollapsed(true).AreaTitle(LOCTEXT("ConstantsLabel", "Constants"))
 										.BodyContent()
 										[
 											SAssignNew(UpdateScriptConstantList, SListView<TSharedPtr<EditorExposedVectorConstant>>)
@@ -323,7 +323,7 @@ void SEmitterWidget::Construct(const FArguments& InArgs)
 											.Padding(4)
 											[
 												SNew(STextBlock)
-												.Text(FString("Spawn Script"))
+												.Text(LOCTEXT("SpawnScriptLabel", "Spawn Script"))
 											]
 											+ SHorizontalBox::Slot()
 											.HAlign(HAlign_Right)
@@ -345,7 +345,7 @@ void SEmitterWidget::Construct(const FArguments& InArgs)
 										.AutoHeight()
 										.Expose(SpawnScriptConstantListSlot)
 										[
-											SNew(SExpandableArea).InitiallyCollapsed(true).AreaTitle(FString("Constants"))
+											SNew(SExpandableArea).InitiallyCollapsed(true).AreaTitle(LOCTEXT("ConstantsLabel", "Constants"))
 											.BodyContent()
 											[
 												SAssignNew(SpawnScriptConstantList, SListView<TSharedPtr<EditorExposedVectorConstant>>)
@@ -379,7 +379,7 @@ void SEmitterWidget::Construct(const FArguments& InArgs)
 										.Padding(4)
 										[
 											SNew(STextBlock)
-											.Text(FString("Spawn rate"))
+											.Text(LOCTEXT("SpawnRateLabel", "Spawn Rate"))
 										]
 										+ SHorizontalBox::Slot()
 											[
@@ -412,7 +412,7 @@ void SEmitterWidget::Construct(const FArguments& InArgs)
 										.AutoWidth()
 										[
 											SNew(STextBlock)
-											.Text(FString("Material"))
+											.Text(LOCTEXT("MaterialLabel", "Material"))
 										]
 										+ SHorizontalBox::Slot()
 											.Padding(4)
@@ -460,7 +460,7 @@ void SEmitterWidget::Construct(const FArguments& InArgs)
 											.AutoWidth()
 											[
 												SNew(STextBlock)
-												.Text(FString("Render as"))
+												.Text(LOCTEXT("RenderAsLabel", "Render as"))
 											]
 											+ SHorizontalBox::Slot()
 											.Padding(4)
@@ -542,7 +542,7 @@ void SEmitterWidget::OnUpdateScriptSelectedFromPicker(UObject *Asset)
 	UpdateScriptConstantListSlot->DetachWidget();
 	(*UpdateScriptConstantListSlot)
 		[
-			SNew(SExpandableArea).InitiallyCollapsed(true).AreaTitle(FString("Constants"))
+			SNew(SExpandableArea).InitiallyCollapsed(true).AreaTitle(LOCTEXT("ConstantsLabel", "Constants"))
 			.BodyContent()
 			[
 				SAssignNew(UpdateScriptConstantList, SListView<TSharedPtr<EditorExposedVectorConstant>>)
@@ -562,7 +562,7 @@ void SEmitterWidget::OnSpawnScriptSelectedFromPicker(UObject *Asset)
 	SpawnScriptConstantListSlot->DetachWidget();
 	(*SpawnScriptConstantListSlot)
 		[
-			SNew(SExpandableArea).InitiallyCollapsed(true).AreaTitle(FString("Constants"))
+			SNew(SExpandableArea).InitiallyCollapsed(true).AreaTitle(LOCTEXT("ConstantsLabel", "Constants"))
 			.BodyContent()
 			[
 				SAssignNew(SpawnScriptConstantList, SListView<TSharedPtr<EditorExposedVectorConstant>>)
diff --git a/Engine/Source/Editor/NiagaraEditor/Private/SNiagaraEffectEditorWidget.h b/Engine/Source/Editor/NiagaraEditor/Private/SNiagaraEffectEditorWidget.h
--- a/Engine/Source/Editor/NiagaraEditor/Private/SNiagaraEffectEditorWidget.h
+++ b/Engine/Source/Editor/NiagaraEditor/Private/SNiagaraEffectEditorWidget.h
@@ -305,28 +305,27 @@ public:
 		return FText::FromString(Txt);
 	}
 
-	FString GetRenderModuleText() const
+	FText GetRenderModuleText() const
 	{
 		if (Emitter->GetProperties()->RenderModuleType == RMT_None)
 		{
-			return FString("None");
+			return NSLOCTEXT("NiagaraEffectEditor", "None", "None");
 		}
 
-		return *RenderModuleList[Emitter->GetProperties()->RenderModuleType-1];
+		return FText::FromString(*RenderModuleList[Emitter->GetProperties()->RenderModuleType-1]);
 	}
 
 	TSharedRef<SWidget>GenerateRenderModuleComboWidget(TSharedPtr<FString> WidgetString )
 	{
-		FString Txt = *WidgetString;
-		return SNew(STextBlock).Text(Txt);
+		return SNew(STextBlock)
+			.Text(FText::FromString(*WidgetString));
 	}
 
-
-	FString GetStatsText() const
+	FText GetStatsText() const
 	{
 		TCHAR Txt[128];
 		FCString::Snprintf(Txt, 128, TEXT("%i Particles  |  %.2f ms  |  %.1f MB"), Emitter->GetNumParticles(), Emitter->GetTotalCPUTime(), Emitter->GetTotalBytesUsed()/(1024.0*1024.0));
-		return FString(Txt);
+		return FText::FromString(Txt);
 	}
 
 	void Construct(const FArguments& InArgs);
diff --git a/Engine/Source/Editor/PListEditor/Private/PListNodeArray.cpp b/Engine/Source/Editor/PListEditor/Private/PListNodeArray.cpp
--- a/Engine/Source/Editor/PListEditor/Private/PListNodeArray.cpp
+++ b/Engine/Source/Editor/PListEditor/Private/PListNodeArray.cpp
@@ -147,14 +147,12 @@ TSharedRef<SWidget> FPListNodeArray::GenerateWidgetForColumn(const FName& Column
 
 	else if(ColumnName == "PListValueColumn")
 	{
-		FString StatusText = NSLOCTEXT("PListEditor", "PListEditorKeyValuePairs", "key/value pairs").ToString();
-
 		return 
 		SNew(SBorder)
 		.BorderImage_Static(&IPListNode::GetOverlayBrushDelegate, AsShared())
 		[
 			SAssignNew(InfoTextWidget, STextBlock)
-			.Text(FString(TEXT("[")) + FString::Printf(TEXT("%i "), GetNumPairs()) + StatusText + FString(TEXT("]")))
+			.Text(FText::Format(NSLOCTEXT("PListEditor", "NumKeyValuePairsFmt", "[{0} key/value pairs]"), FText::AsNumber(GetNumPairs())))
 		];
 	}
 
@@ -200,10 +198,9 @@ FString FPListNodeArray::ToXML(const int32 indent, bool bOutputKey)
 void FPListNodeArray::Refresh()
 {
 	// Refresh display of # of child widgets (key/value pairs)
-	FString StatusText = NSLOCTEXT("PListEditor", "PListEditorKeyValuePairs", "key/value pairs").ToString();
 	if(InfoTextWidget.IsValid())
 	{
-		InfoTextWidget->SetText(FString(TEXT("[")) + FString::Printf(TEXT("%i "), GetNumPairs()) + StatusText + FString(TEXT("]")));
+		InfoTextWidget->SetText(FText::Format(NSLOCTEXT("PListEditor", "NumKeyValuePairsFmt", "[{0} key/value pairs]"), FText::AsNumber(GetNumPairs())));
 	}
 
 	// Refresh all children and set their indices
diff --git a/Engine/Source/Editor/PListEditor/Private/PListNodeDictionary.cpp b/Engine/Source/Editor/PListEditor/Private/PListNodeDictionary.cpp
--- a/Engine/Source/Editor/PListEditor/Private/PListNodeDictionary.cpp
+++ b/Engine/Source/Editor/PListEditor/Private/PListNodeDictionary.cpp
@@ -135,14 +135,12 @@ TSharedRef<SWidget> FPListNodeDictionary::GenerateWidgetForColumn(const FName& C
 
 	else if(ColumnName == "PListValueColumn")
 	{
-		FString StatusText = NSLOCTEXT("PListEditor", "PListEditorKeyValuePairs", "key/value pairs").ToString();
-
 		return
 		SNew(SBorder)
 		.BorderImage_Static(&IPListNode::GetOverlayBrushDelegate, AsShared())
 		[
 			SAssignNew(InfoTextWidget, STextBlock)
-			.Text(FString(TEXT("[")) + FString::Printf(TEXT("%i "), GetNumPairs()) + StatusText + FString(TEXT("]")))
+			.Text(FText::Format(NSLOCTEXT("PListEditor", "NumKeyValuePairsFmt", "[{0} key/value pairs]"), FText::AsNumber(GetNumPairs())))
 		];
 	}
 
@@ -181,10 +179,9 @@ FString FPListNodeDictionary::ToXML(const int32 indent, bool bOutputKey)
 void FPListNodeDictionary::Refresh()
 {
 	// Update the display of the number of key/value pairs
-	FString StatusText = NSLOCTEXT("PListEditor", "PListEditorKeyValuePairs", "key/value pairs").ToString();
 	if(InfoTextWidget.IsValid())
 	{
-		InfoTextWidget->SetText(FString(TEXT("[")) + FString::Printf(TEXT("%i "), GetNumPairs()) + StatusText + FString(TEXT("]")));
+		InfoTextWidget->SetText(FText::Format(NSLOCTEXT("PListEditor", "NumKeyValuePairsFmt", "[{0} key/value pairs]"), FText::AsNumber(GetNumPairs())));
 	}
 
 	// Refresh all children
diff --git a/Engine/Source/Editor/PListEditor/Private/PListNodeFile.cpp b/Engine/Source/Editor/PListEditor/Private/PListNodeFile.cpp
--- a/Engine/Source/Editor/PListEditor/Private/PListNodeFile.cpp
+++ b/Engine/Source/Editor/PListEditor/Private/PListNodeFile.cpp
@@ -45,9 +45,6 @@ bool FPListNodeFile::UsesColumns()
 /** Generates a widget for a TableView row */
 TSharedRef<ITableRow> FPListNodeFile::GenerateWidget(const TSharedRef<STableViewBase>& OwnerTable)
 {
-	FString FileLabel = NSLOCTEXT("PListEditor", "PListEditorFileLabel", "file").ToString();
-	FString StatusText = NSLOCTEXT("PListEditor", "PListEditorKeyValuePairs", "key/value pairs").ToString();
-
 	return
 	SNew(STableRow<TSharedPtr<ITableRow> >, OwnerTable)
 	.Content()
@@ -59,7 +56,7 @@ TSharedRef<ITableRow> FPListNodeFile::GenerateWidget(const TSharedRef<STableView
 		.Padding(2, 1)
 		[
 			SAssignNew(TextWidget, STextBlock)
-			.Text(FileLabel + FString(TEXT(" [")) + FString::Printf(TEXT("%i "), GetNumPairs()) + StatusText + FString(TEXT("]")))
+			.Text(FText::Format(NSLOCTEXT("PListEditor", "FileAndNumKeyValuePairsFmt", "file [{0} key/value pairs]"), FText::AsNumber(GetNumPairs())))
 		]
 	];
 }
@@ -101,11 +98,9 @@ void FPListNodeFile::Refresh()
 	}
 
 	// Calculate how many contained key/value pairs exist
-	FString FileLabel = NSLOCTEXT("PListEditor", "PListEditorFileLabel", "file").ToString();
-	FString StatusText = NSLOCTEXT("PListEditor", "PListEditorKeyValuePairs", "key/value pairs").ToString();
 	if(TextWidget.IsValid())
 	{
-		TextWidget->SetText(FileLabel + FString(TEXT(" [")) + FString::Printf(TEXT("%i "), GetNumPairs()) + StatusText + FString(TEXT("]")));
+		TextWidget->SetText(FText::Format(NSLOCTEXT("PListEditor", "FileAndNumKeyValuePairsFmt", "file [{0} key/value pairs]"), FText::AsNumber(GetNumPairs())));
 	}
 }
 
diff --git a/Engine/Source/Editor/PackagesDialog/Private/SPackagesDialog.cpp b/Engine/Source/Editor/PackagesDialog/Private/SPackagesDialog.cpp
--- a/Engine/Source/Editor/PackagesDialog/Private/SPackagesDialog.cpp
+++ b/Engine/Source/Editor/PackagesDialog/Private/SPackagesDialog.cpp
@@ -380,7 +380,7 @@ TSharedRef<SWidget> SPackagesDialog::GenerateWidgetForItemAndColumn( TSharedPtr<
 			.Padding(RowPadding)
 			[
 				SNew(STextBlock)
-				.Text(PackageName)
+				.Text(FText::FromString(PackageName))
 				.IsEnabled(!Item->IsDisabled())
 			];
 	}
@@ -391,7 +391,7 @@ TSharedRef<SWidget> SPackagesDialog::GenerateWidgetForItemAndColumn( TSharedPtr<
 			.Padding(RowPadding)
 			[
 				SNew(STextBlock)
-				.Text(PackageType)
+				.Text(FText::FromString(PackageType))
 				.IsEnabled(!Item->IsDisabled())
 				.ColorAndOpacity(PackageColor)
 			];
diff --git a/Engine/Source/Editor/Persona/Private/AnimGraphNodeDetails.cpp b/Engine/Source/Editor/Persona/Private/AnimGraphNodeDetails.cpp
--- a/Engine/Source/Editor/Persona/Private/AnimGraphNodeDetails.cpp
+++ b/Engine/Source/Editor/Persona/Private/AnimGraphNodeDetails.cpp
@@ -793,7 +793,7 @@ void FPlayerTreeViewEntry::GenerateNameWidget(TSharedPtr<SHorizontalBox> Box)
 		[
 			SNew(STextBlock)
 			.Font(FSlateFontInfo(FPaths::EngineContentDir() / TEXT("Slate/Fonts/Roboto-Bold.ttf"), 10))
-			.Text(EntryName)
+			.Text(FText::FromString(EntryName))
 		];
 }
 
diff --git a/Engine/Source/Editor/Persona/Private/AnimationEditorViewportClient.cpp b/Engine/Source/Editor/Persona/Private/AnimationEditorViewportClient.cpp
--- a/Engine/Source/Editor/Persona/Private/AnimationEditorViewportClient.cpp
+++ b/Engine/Source/Editor/Persona/Private/AnimationEditorViewportClient.cpp
@@ -2237,13 +2237,15 @@ float FAnimationViewportClient::GetWindStrengthSliderValue() const
 	return 0;
 }
 
-FString FAnimationViewportClient::GetWindStrengthLabel() const
+FText FAnimationViewportClient::GetWindStrengthLabel() const
 {
 	//Clamp slide value so that minimum value displayed is 0.00 and maximum is 1.0
 	float SliderValue = FMath::Clamp<float>(GetWindStrengthSliderValue(), 0.0f, 1.0f);
 
-	FString Value = FString::Printf(TEXT("%.2f%"),SliderValue);
-	return Value;
+	static const FNumberFormattingOptions FormatOptions = FNumberFormattingOptions()
+		.SetMinimumFractionalDigits(2)
+		.SetMaximumFractionalDigits(2);
+	return FText::AsNumber(SliderValue, &FormatOptions);
 }
 
 void FAnimationViewportClient::SetGravityScale( float SliderPos )
@@ -2264,13 +2266,15 @@ float FAnimationViewportClient::GetGravityScaleSliderValue() const
 	return GravityScaleSliderValue;
 }
 
-FString FAnimationViewportClient::GetGravityScaleLabel() const
+FText FAnimationViewportClient::GetGravityScaleLabel() const
 {
 	//Clamp slide value so that minimum value displayed is 0.00 and maximum is 4.0
 	float SliderValue = FMath::Clamp<float>(GetGravityScaleSliderValue()*4, 0.0f, 4.0f);
 
-	FString Value = FString::Printf(TEXT("%.2f%"),SliderValue);
-	return Value;
+	static const FNumberFormattingOptions FormatOptions = FNumberFormattingOptions()
+		.SetMinimumFractionalDigits(2)
+		.SetMaximumFractionalDigits(2);
+	return FText::AsNumber(SliderValue, &FormatOptions);
 }
 
 void FAnimationViewportClient::ToggleShowNormals()
diff --git a/Engine/Source/Editor/Persona/Private/AnimationEditorViewportClient.h b/Engine/Source/Editor/Persona/Private/AnimationEditorViewportClient.h
--- a/Engine/Source/Editor/Persona/Private/AnimationEditorViewportClient.h
+++ b/Engine/Source/Editor/Persona/Private/AnimationEditorViewportClient.h
@@ -204,13 +204,13 @@ public:
 	float GetWindStrengthSliderValue() const;
 
 	/** Function to get wind strength label */
-	FString GetWindStrengthLabel() const;
+	FText GetWindStrengthLabel() const;
 	/** Function to set gravity scale (0.0 - 4.0) */
 	void SetGravityScale(float SliderPos);
 	/** Function to get slide value used to represent gravity scale */
 	float GetGravityScaleSliderValue() const;
 	/** Function to get gravity scale label */
-	FString GetGravityScaleLabel() const;
+	FText GetGravityScaleLabel() const;
 	/** Function to set mesh stat drawing state */
 	void OnSetShowMeshStats(int32 ShowMode);
 	/** Whether or not mesh stats are being displayed */
diff --git a/Engine/Source/Editor/Persona/Private/AnimationMode/AnimationMode.cpp b/Engine/Source/Editor/Persona/Private/AnimationMode/AnimationMode.cpp
--- a/Engine/Source/Editor/Persona/Private/AnimationMode/AnimationMode.cpp
+++ b/Engine/Source/Editor/Persona/Private/AnimationMode/AnimationMode.cpp
@@ -36,15 +36,15 @@ public:
 		return (GetObjectToObserve() != NULL) ? EVisibility::Visible : EVisibility::Collapsed;
 	}
 
-	virtual FString GetAssetDisplayName() const
+	virtual FText GetAssetDisplayName() const
 	{
 		if (UObject* Object = GetObjectToObserve())
 		{
-			return Object->GetName();
+			return FText::FromString(Object->GetName());
 		}
 		else
 		{
-			return FString();
+			return FText::GetEmpty();
 		}
 	}
 
diff --git a/Engine/Source/Editor/Persona/Private/MeshMode/SAdditionalMeshesEditor.cpp b/Engine/Source/Editor/Persona/Private/MeshMode/SAdditionalMeshesEditor.cpp
--- a/Engine/Source/Editor/Persona/Private/MeshMode/SAdditionalMeshesEditor.cpp
+++ b/Engine/Source/Editor/Persona/Private/MeshMode/SAdditionalMeshesEditor.cpp
@@ -70,7 +70,7 @@ void SAdditionalMeshItem::Construct(const FArguments& InArgs, USkeletalMeshCompo
 			.AutoWidth()
 			[
 				SNew(STextBlock)
-				.Text(MeshComponent->SkeletalMesh->GetName())
+				.Text(FText::FromString(MeshComponent->SkeletalMesh->GetName()))
 			]
 			+SHorizontalBox::Slot()
 			.AutoWidth()
diff --git a/Engine/Source/Editor/Persona/Private/Persona.cpp b/Engine/Source/Editor/Persona/Private/Persona.cpp
--- a/Engine/Source/Editor/Persona/Private/Persona.cpp
+++ b/Engine/Source/Editor/Persona/Private/Persona.cpp
@@ -1027,14 +1027,14 @@ void FPersona::SetPreviewVertexAnim(UVertexAnimation* VertexAnim)
 	}
 }
 
-void FPersona::UpdateSelectionDetails(UObject* Object, const FString& ForcedTitle)
+void FPersona::UpdateSelectionDetails(UObject* Object, const FText& ForcedTitle)
 {
 	Inspector->ShowDetailsForSingleObject(Object, SKismetInspector::FShowDetailsOptions(ForcedTitle));
 }
 
 void FPersona::SetDetailObject(UObject* Obj)
 {
-	FString ForcedTitle = (Obj != NULL) ? *Obj->GetName() : FString();
+	FText ForcedTitle = (Obj != NULL) ? FText::FromString(Obj->GetName()) : FText::GetEmpty();
 	UpdateSelectionDetails(Obj, ForcedTitle);
 }
 
@@ -1249,7 +1249,7 @@ void FPersona::StartEditingDefaults(bool bAutoFocus, bool bForceRefresh)
 	{
 		if (PreviewComponent->AnimScriptInstance != NULL)
 		{
-			PreviewEditor->ShowDetailsForSingleObject(PreviewComponent->AnimScriptInstance, SKismetInspector::FShowDetailsOptions(TEXT("Preview settings")));
+			PreviewEditor->ShowDetailsForSingleObject(PreviewComponent->AnimScriptInstance, SKismetInspector::FShowDetailsOptions(LOCTEXT("PreviewSettingsTitle", "Preview settings")));
 		}
 	}
 	if (bAutoFocus)
@@ -1649,9 +1649,9 @@ void FPersona::Compile()
 	}
 }
 
-FString FPersona::GetDefaultEditorTitle()
+FText FPersona::GetDefaultEditorTitle()
 {
-	return NSLOCTEXT("Kismet", "PreviewParamatersTabTitle", "Preview Parameters").ToString();
+	return NSLOCTEXT("Kismet", "PreviewParamatersTabTitle", "Preview Parameters");
 }
 
 FName FPersona::GetToolkitContextFName() const
@@ -1935,7 +1935,7 @@ FGraphAppearanceInfo FPersona::GetGraphAppearance() const
 
 	if ( GetBlueprintObj()->IsA(UAnimBlueprint::StaticClass()) )
 	{
-		AppearanceInfo.CornerText = LOCTEXT("AppearanceCornerText_Animation", "ANIMATION").ToString();
+		AppearanceInfo.CornerText = LOCTEXT("AppearanceCornerText_Animation", "ANIMATION");
 	}
 
 	return AppearanceInfo;
@@ -2414,7 +2414,7 @@ FReply FPersona::OnClickEditMesh()
 	USkeletalMesh* PreviewMesh = TargetSkeleton->GetPreviewMesh();
 	if(PreviewMesh)
 	{
-		UpdateSelectionDetails(PreviewMesh, *PreviewMesh->GetName());
+		UpdateSelectionDetails(PreviewMesh, FText::FromString(PreviewMesh->GetName()));
 	}
 	return FReply::Handled();
 }
@@ -2958,13 +2958,13 @@ bool FPersona::IsEditable(UEdGraph* InGraph) const
 	return bEditable;
 }
 
-FString FPersona::GetGraphDecorationString(UEdGraph* InGraph) const
+FText FPersona::GetGraphDecorationString(UEdGraph* InGraph) const
 {
 	if (!IsGraphInCurrentBlueprint(InGraph))
 	{
-		return LOCTEXT("PersonaExternalGraphDecoration", " Parent Graph Preview").ToString();
+		return LOCTEXT("PersonaExternalGraphDecoration", " Parent Graph Preview");
 	}
-	return TEXT("");
+	return FText::GetEmpty();
 }
 
 void FPersona::ValidatePreviewAttachedAssets(USkeletalMesh* PreviewSkeletalMesh)
diff --git a/Engine/Source/Editor/Persona/Private/Persona.h b/Engine/Source/Editor/Persona/Private/Persona.h
--- a/Engine/Source/Editor/Persona/Private/Persona.h
+++ b/Engine/Source/Editor/Persona/Private/Persona.h
@@ -57,7 +57,7 @@ public:
 	void SetPreviewVertexAnim(UVertexAnimation* VertexAnim);
 
 	/** Update the inspector that displays information about the current selection*/
-	void UpdateSelectionDetails(UObject* Object, const FString& ForcedTitle);
+	void UpdateSelectionDetails(UObject* Object, const FText& ForcedTitle);
 
 	void SetDetailObject(UObject* Obj);
 
@@ -219,7 +219,7 @@ protected:
 	virtual void StartEditingDefaults(bool bAutoFocus, bool bForceRefresh = false) override;
 	virtual void Compile() override;
 	virtual void OnGraphEditorFocused(const TSharedRef<class SGraphEditor>& InGraphEditor) override;
-	virtual FString GetDefaultEditorTitle() override;
+	virtual FText GetDefaultEditorTitle() override;
 	virtual void OnConvertToSequenceEvaluator() override;
 	virtual void OnConvertToSequencePlayer() override;
 	virtual void OnConvertToBlendSpaceEvaluator() override;
@@ -230,7 +230,7 @@ protected:
 	virtual void CreateDefaultTabContents(const TArray<UBlueprint*>& InBlueprints) override;
 	virtual FGraphAppearanceInfo GetGraphAppearance() const override;
 	virtual bool IsEditable(UEdGraph* InGraph) const override;
-	virtual FString GetGraphDecorationString(UEdGraph* InGraph) const override;
+	virtual FText GetGraphDecorationString(UEdGraph* InGraph) const override;
 	// End of FBlueprintEditor interface
 
 	// IAssetEditorInstance interface
diff --git a/Engine/Source/Editor/Persona/Private/PersonaMeshDetails.cpp b/Engine/Source/Editor/Persona/Private/PersonaMeshDetails.cpp
--- a/Engine/Source/Editor/Persona/Private/PersonaMeshDetails.cpp
+++ b/Engine/Source/Editor/Persona/Private/PersonaMeshDetails.cpp
@@ -1490,7 +1490,7 @@ TSharedRef<SUniformGridPanel> FPersonaMeshDetails::MakeApexDetailsWidget(int32 A
 				[
 					SNew(STextBlock)
 					.Font(DetailFontInfo)
-					.Text(FString::Printf( TEXT("%d"),SubmeshInfos[i].SubmeshIndex))
+					.Text(FText::AsNumber(SubmeshInfos[i].SubmeshIndex))
 				];
 
 				if(i == 0)
@@ -1500,7 +1500,7 @@ TSharedRef<SUniformGridPanel> FPersonaMeshDetails::MakeApexDetailsWidget(int32 A
 					[
 						SNew(STextBlock)
 						.Font(DetailFontInfo)
-						.Text(FString::Printf( TEXT("%d"),SubmeshInfos[i].SimulVertexCount))
+						.Text(FText::AsNumber(SubmeshInfos[i].SimulVertexCount))
 					];
 				}
 				else
@@ -1510,7 +1510,7 @@ TSharedRef<SUniformGridPanel> FPersonaMeshDetails::MakeApexDetailsWidget(int32 A
 					[
 						SNew(STextBlock)
 						.Font(DetailFontInfo)
-						.Text(FString::Printf( TEXT("Shared") ))
+						.Text(LOCTEXT("SharedLabel", "Shared"))
 					];
 				}
 
@@ -1519,7 +1519,7 @@ TSharedRef<SUniformGridPanel> FPersonaMeshDetails::MakeApexDetailsWidget(int32 A
 					[
 						SNew(STextBlock)
 						.Font(DetailFontInfo)
-						.Text(FString::Printf( TEXT("%d"),SubmeshInfos[i].VertexCount))
+						.Text(FText::AsNumber(SubmeshInfos[i].VertexCount))
 					];
 
 				Grid->AddSlot(3, RowNumber)
@@ -1527,7 +1527,7 @@ TSharedRef<SUniformGridPanel> FPersonaMeshDetails::MakeApexDetailsWidget(int32 A
 				[
 					SNew(STextBlock)
 					.Font(DetailFontInfo)
-					.Text(FString::Printf( TEXT("%d"),SubmeshInfos[i].FixedVertexCount))
+					.Text(FText::AsNumber(SubmeshInfos[i].FixedVertexCount))
 				];
 
 				Grid->AddSlot(4, RowNumber)
@@ -1535,7 +1535,7 @@ TSharedRef<SUniformGridPanel> FPersonaMeshDetails::MakeApexDetailsWidget(int32 A
 				[
 					SNew(STextBlock)
 					.Font(DetailFontInfo)
-					.Text(FString::Printf( TEXT("%d"),SubmeshInfos[i].TriangleCount))
+					.Text(FText::AsNumber(SubmeshInfos[i].TriangleCount))
 				];
 
 				RowNumber++;
diff --git a/Engine/Source/Editor/Persona/Private/SAnimEditorBase.cpp b/Engine/Source/Editor/Persona/Private/SAnimEditorBase.cpp
--- a/Engine/Source/Editor/Persona/Private/SAnimEditorBase.cpp
+++ b/Engine/Source/Editor/Persona/Private/SAnimEditorBase.cpp
@@ -216,15 +216,15 @@ void SAnimEditorBase::ClearDetailsView()
 	PersonaPtr.Pin()->SetDetailObject(NULL);
 }
 
-FString SAnimEditorBase::GetEditorObjectName() const
+FText SAnimEditorBase::GetEditorObjectName() const
 {
 	if (GetEditorObject() != NULL)
 	{
-		return GetEditorObject()->GetName();
+		return FText::FromString(GetEditorObject()->GetName());
 	}
 	else
 	{
-		return LOCTEXT("NoEditorObject", "No Editor Object").ToString();
+		return LOCTEXT("NoEditorObject", "No Editor Object");
 	}
 }
 
@@ -272,7 +272,7 @@ void SAnimEditorBase::OnSelectionChanged(const FGraphPanelSelectionSet& Selected
 	if (SelectedItems.Num() == 0)
 	{
 		// Edit the sequence
-		PersonaPtr.Pin()->UpdateSelectionDetails(GetEditorObject(), LOCTEXT("Edit Sequence", "Edit Sequence").ToString());
+		PersonaPtr.Pin()->UpdateSelectionDetails(GetEditorObject(), LOCTEXT("Edit Sequence", "Edit Sequence"));
 	}
 	else
 	{
@@ -306,7 +306,7 @@ void SAnimEditorBase::SetInputViewRange(float InViewMinInput, float InViewMaxInp
 	ViewMinInput = FMath::Max<float>(InViewMinInput, 0.f);
 }
 
-FString SAnimEditorBase::GetCurrentSequenceTime() const
+FText SAnimEditorBase::GetCurrentSequenceTime() const
 {
 	UAnimSingleNodeInstance * PreviewInstance = GetPreviewInstance();
 	float CurTime = 0.f;
@@ -317,11 +317,13 @@ FString SAnimEditorBase::GetCurrentSequenceTime() const
 		CurTime = PreviewInstance->CurrentTime;
 	}
 
-	const FString Fraction = FString::Printf(TEXT("%0.3f / %0.3f"), CurTime, TotalTime);
-	return FText::Format( LOCTEXT("FractionSeconds", "{0} (second(s))"), FText::FromString(Fraction) ).ToString();
+	static const FNumberFormattingOptions FractionNumberFormat = FNumberFormattingOptions()
+		.SetMinimumFractionalDigits(3)
+		.SetMaximumFractionalDigits(3);
+	return FText::Format(LOCTEXT("FractionSecondsFmt", "{0} / {1} (second(s))"), FText::AsNumber(CurTime, &FractionNumberFormat), FText::AsNumber(TotalTime, &FractionNumberFormat));
 }
 
-FString SAnimEditorBase::GetCurrentPercentage() const
+FText SAnimEditorBase::GetCurrentPercentage() const
 {
 	UAnimSingleNodeInstance * PreviewInstance = GetPreviewInstance();
 	float Percentage = 0.f;
@@ -330,10 +332,13 @@ FString SAnimEditorBase::GetCurrentPercentage() const
 		Percentage = PreviewInstance->CurrentTime / GetEditorObject()->SequenceLength;
 	}
 
-	return FString::Printf(TEXT("%0.2f %%"), Percentage*100.f);
+	static const FNumberFormattingOptions PercentNumberFormat = FNumberFormattingOptions()
+		.SetMinimumFractionalDigits(2)
+		.SetMaximumFractionalDigits(2);
+	return FText::AsPercent(Percentage, &PercentNumberFormat);
 }
 
-FString SAnimEditorBase::GetCurrentFrame() const
+FText SAnimEditorBase::GetCurrentFrame() const
 {
 	UAnimSingleNodeInstance * PreviewInstance = GetPreviewInstance();
 	float Percentage = 0.f;
@@ -344,8 +349,10 @@ FString SAnimEditorBase::GetCurrentFrame() const
 		Percentage = PreviewInstance->CurrentTime/GetEditorObject()->SequenceLength;
 	}
 
-	const FString Fraction = FString::Printf(TEXT("%0.2f / %d"), NumFrames*Percentage, (int32)NumFrames);
-	return FText::Format( LOCTEXT("FractionKeys", "{0} (key(s))"), FText::FromString(Fraction) ).ToString();
+	static const FNumberFormattingOptions FractionNumberFormat = FNumberFormattingOptions()
+		.SetMinimumFractionalDigits(2)
+		.SetMaximumFractionalDigits(2);
+	return FText::Format(LOCTEXT("FractionKeysFmt", "{0} / {1} (key(s))"), FText::AsNumber(NumFrames * Percentage, &FractionNumberFormat), FText::AsNumber((int32)NumFrames));
 }
 
 #undef LOCTEXT_NAMESPACE
diff --git a/Engine/Source/Editor/Persona/Private/SAnimEditorBase.h b/Engine/Source/Editor/Persona/Private/SAnimEditorBase.h
--- a/Engine/Source/Editor/Persona/Private/SAnimEditorBase.h
+++ b/Engine/Source/Editor/Persona/Private/SAnimEditorBase.h
@@ -85,7 +85,7 @@ protected:
 	virtual UAnimSequenceBase* GetEditorObject() const PURE_VIRTUAL(SAnimEditorBase::GetEditorObject, return NULL;);
 
 	/** Get Name of Object being edited **/
-	FString GetEditorObjectName() const;
+	FText GetEditorObjectName() const;
 
 	/** Get Min/Max Input of value **/
 	float ViewMinInput;
@@ -95,9 +95,9 @@ protected:
 	class UAnimSingleNodeInstance* GetPreviewInstance() const;
 
 	/** For the information section **/
-	FString GetCurrentSequenceTime() const;
-	FString GetCurrentPercentage() const;
-	FString GetCurrentFrame() const;
+	FText GetCurrentSequenceTime() const;
+	FText GetCurrentPercentage() const;
+	FText GetCurrentFrame() const;
 
 	/** Recalculate sequence length after modifying */
 	void RecalculateSequenceLength();
diff --git a/Engine/Source/Editor/Persona/Private/SAnimationBlendSpace.cpp b/Engine/Source/Editor/Persona/Private/SAnimationBlendSpace.cpp
--- a/Engine/Source/Editor/Persona/Private/SAnimationBlendSpace.cpp
+++ b/Engine/Source/Editor/Persona/Private/SAnimationBlendSpace.cpp
@@ -1419,10 +1419,16 @@ void SBlendSpaceEditor::UpdateBlendParameters()
 		ParameterYName = TEXT("Y");
 	}
 
-	ParameterX_Min->SetText(FString::Printf(TEXT("%s[%0.2f]"), *ParameterXName, BlendParamX.Min));
-	ParameterX_Max->SetText(FString::Printf(TEXT("%s[%0.2f]"), *ParameterXName, BlendParamX.Max));
-	ParameterY_Min->SetText(FString::Printf(TEXT("%s\n[%0.2f]"), *ParameterYName, BlendParamY.Min));
-	ParameterY_Max->SetText(FString::Printf(TEXT("%s\n[%0.2f]"), *ParameterYName, BlendParamY.Max));
+	static const FNumberFormattingOptions FormatOptions = FNumberFormattingOptions()
+		.SetMinimumFractionalDigits(2)
+		.SetMaximumFractionalDigits(2);
+	static const FText XMinMaxFmt = FText::FromString(TEXT("{0}[{1}]"));
+	static const FText YMinMaxFmt = FText::FromString(TEXT("{0}\n[{1}]"));
+
+	ParameterX_Min->SetText(FText::Format(XMinMaxFmt, FText::FromString(ParameterXName), FText::AsNumber(BlendParamX.Min, &FormatOptions)));
+	ParameterX_Max->SetText(FText::Format(XMinMaxFmt, FText::FromString(ParameterXName), FText::AsNumber(BlendParamX.Max, &FormatOptions)));
+	ParameterY_Min->SetText(FText::Format(YMinMaxFmt, FText::FromString(ParameterYName), FText::AsNumber(BlendParamY.Min, &FormatOptions)));
+	ParameterY_Max->SetText(FText::Format(YMinMaxFmt, FText::FromString(ParameterYName), FText::AsNumber(BlendParamY.Max, &FormatOptions)));
 
 	BlendSpaceWidget->ResampleData();
 }
diff --git a/Engine/Source/Editor/Persona/Private/SAnimationBlendSpace1D.cpp b/Engine/Source/Editor/Persona/Private/SAnimationBlendSpace1D.cpp
--- a/Engine/Source/Editor/Persona/Private/SAnimationBlendSpace1D.cpp
+++ b/Engine/Source/Editor/Persona/Private/SAnimationBlendSpace1D.cpp
@@ -629,18 +629,22 @@ void SBlendSpaceEditor1D::UpdateBlendParameters()
 		ParameterName = TEXT("X");
 	}
 
-	FString LabelFormat;
+	FText LabelFormat;
 	if(GetBlendSpace()->bDisplayEditorVertically)
 	{
-		LabelFormat = TEXT("%s\n[%0.2f]");
+		LabelFormat = FText::FromString(TEXT("{0}\n[{1}]"));
 	}
 	else
 	{
-		LabelFormat = TEXT("%s[%0.2f]");
+		LabelFormat = FText::FromString(TEXT("{0}[{1}]"));
 	}
 
-	Parameter_Min->SetText(FString::Printf(*LabelFormat, *ParameterName, BlendParam.Min));
-	Parameter_Max->SetText(FString::Printf(*LabelFormat, *ParameterName, BlendParam.Max));
+	static const FNumberFormattingOptions NumberFormat = FNumberFormattingOptions()
+		.SetMinimumFractionalDigits(2)
+		.SetMaximumFractionalDigits(2);
+
+	Parameter_Min->SetText(FText::Format(LabelFormat, FText::FromString(ParameterName), FText::AsNumber(BlendParam.Min, &NumberFormat)));
+	Parameter_Max->SetText(FText::Format(LabelFormat, FText::FromString(ParameterName), FText::AsNumber(BlendParam.Max, &NumberFormat)));
 
 	BlendSpaceWidget->ResampleData();
 }
diff --git a/Engine/Source/Editor/Persona/Private/SAnimationBlendSpaceBase.cpp b/Engine/Source/Editor/Persona/Private/SAnimationBlendSpaceBase.cpp
--- a/Engine/Source/Editor/Persona/Private/SAnimationBlendSpaceBase.cpp
+++ b/Engine/Source/Editor/Persona/Private/SAnimationBlendSpaceBase.cpp
@@ -178,7 +178,7 @@ public:
 				.AutoWidth()
 				[
 					SNew(STextBlock)
-					.Text( Param.DisplayName )
+					.Text( FText::FromString(Param.DisplayName) )
 				]
 				+SHorizontalBox::Slot()
 				.AutoWidth()
@@ -519,7 +519,7 @@ void SBlendSpaceSamplesWidget::ConstructSamplesPanel()
 				[
 					SNew(STextBlock)
 					.Font( GetSmallLayoutFont() )
-					.Text(FString::Printf(TEXT("%3d] "), I+1))
+					.Text(FText::Format(LOCTEXT("BlendSpaceSamplesLabel", "{0}]"), FText::AsNumber(I+1)))
 				]
 				+SHorizontalBox::Slot() .FillWidth(1) .HAlign(HAlign_Left)
 				[
diff --git a/Engine/Source/Editor/Persona/Private/SAnimationBlendSpaceBase.h b/Engine/Source/Editor/Persona/Private/SAnimationBlendSpaceBase.h
--- a/Engine/Source/Editor/Persona/Private/SAnimationBlendSpaceBase.h
+++ b/Engine/Source/Editor/Persona/Private/SAnimationBlendSpaceBase.h
@@ -389,9 +389,9 @@ protected:
 	/* Notifies the user of failures */
 	void NotifyUser( FNotificationInfo& Info );
 
-	FString GetBlendSpaceDisplayName() const
+	FText GetBlendSpaceDisplayName() const
 	{
-		return BlendSpace->GetName();
+		return FText::FromString(BlendSpace->GetName());
 	}
 
 	/** Pointer back to the Persona that owns us */
diff --git a/Engine/Source/Editor/Persona/Private/SAnimationEditorViewport.cpp b/Engine/Source/Editor/Persona/Private/SAnimationEditorViewport.cpp
--- a/Engine/Source/Editor/Persona/Private/SAnimationEditorViewport.cpp
+++ b/Engine/Source/Editor/Persona/Private/SAnimationEditorViewport.cpp
@@ -227,7 +227,7 @@ bool SAnimationEditorViewportTabBody::CanUseGizmos() const
 	return false;
 }
 
-FString SAnimationEditorViewportTabBody::GetDisplayString() const
+FText SAnimationEditorViewportTabBody::GetDisplayString() const
 {
 	class UDebugSkelMeshComponent* Component = PersonaPtr.Pin()->PreviewComponent;
 	FString TargetSkeletonName = TargetSkeleton ?TargetSkeleton->GetName() : FName(NAME_None).ToString();
@@ -236,36 +236,36 @@ FString SAnimationEditorViewportTabBody::GetDisplayString() const
 	{
 		if ( Component->bForceRefpose )
 		{
-			return LOCTEXT("ReferencePose", "Reference pose").ToString();
+			return LOCTEXT("ReferencePose", "Reference pose");
 		}
 		else if (Component->IsPreviewOn())
 		{
-			return FText::Format( LOCTEXT("Previewing", "Previewing {0}"), FText::FromString(Component->GetPreviewText()) ).ToString();
+			return FText::Format( LOCTEXT("Previewing", "Previewing {0}"), FText::FromString(Component->GetPreviewText()) );
 		}
 		else if (Component->AnimBlueprintGeneratedClass != NULL)
 		{
 			const bool bWarnAboutBoneManip = !PersonaPtr.Pin()->IsModeCurrent(FPersonaModes::AnimBlueprintEditMode);
 			if (bWarnAboutBoneManip)
 			{
-				return FText::Format(LOCTEXT("PreviewingAnimBP_WarnDisabled", "Previewing {0}. \nBone manipulation is disabled in this mode. "), FText::FromString(Component->AnimBlueprintGeneratedClass->GetName())).ToString();
+				return FText::Format(LOCTEXT("PreviewingAnimBP_WarnDisabled", "Previewing {0}. \nBone manipulation is disabled in this mode. "), FText::FromString(Component->AnimBlueprintGeneratedClass->GetName()));
 			}
 			else
 			{
-				return FText::Format(LOCTEXT("PreviewingAnimBP", "Previewing {0}"), FText::FromString(Component->AnimBlueprintGeneratedClass->GetName())).ToString();
+				return FText::Format(LOCTEXT("PreviewingAnimBP", "Previewing {0}"), FText::FromString(Component->AnimBlueprintGeneratedClass->GetName()));
 			}
 		}
 		else if (Component->SkeletalMesh == NULL)
 		{
-			return FText::Format( LOCTEXT("NoMeshFound", "No skeletal mesh found for skeleton '{0}'"), FText::FromString(TargetSkeletonName) ).ToString();
+			return FText::Format( LOCTEXT("NoMeshFound", "No skeletal mesh found for skeleton '{0}'"), FText::FromString(TargetSkeletonName) );
 		}
 		else 
 		{
-			return LOCTEXT("NothingToPlay", "Nothing to play").ToString();
+			return LOCTEXT("NothingToPlay", "Nothing to play");
 		}
 	}
 	else
 	{
-		return FText::Format( LOCTEXT("NoMeshFound", "No skeletal mesh found for skeleton '{0}'"), FText::FromString(TargetSkeletonName) ).ToString();
+		return FText::Format( LOCTEXT("NoMeshFound", "No skeletal mesh found for skeleton '{0}'"), FText::FromString(TargetSkeletonName) );
 	}
 }
 
@@ -1549,7 +1549,7 @@ float SAnimationEditorViewportTabBody::GetWindStrengthSliderValue() const
 	return GetAnimationViewportClient()->GetWindStrengthSliderValue();
 }
 
-FString SAnimationEditorViewportTabBody::GetWindStrengthLabel() const
+FText SAnimationEditorViewportTabBody::GetWindStrengthLabel() const
 {
 	return GetAnimationViewportClient()->GetWindStrengthLabel();
 }
@@ -1634,7 +1634,7 @@ float SAnimationEditorViewportTabBody::GetGravityScaleSliderValue() const
 	return GetAnimationViewportClient()->GetGravityScaleSliderValue();
 }
 
-FString SAnimationEditorViewportTabBody::GetGravityScaleLabel() const
+FText SAnimationEditorViewportTabBody::GetGravityScaleLabel() const
 {
 	return GetAnimationViewportClient()->GetGravityScaleLabel();
 }
diff --git a/Engine/Source/Editor/Persona/Private/SAnimationEditorViewport.h b/Engine/Source/Editor/Persona/Private/SAnimationEditorViewport.h
--- a/Engine/Source/Editor/Persona/Private/SAnimationEditorViewport.h
+++ b/Engine/Source/Editor/Persona/Private/SAnimationEditorViewport.h
@@ -148,7 +148,7 @@ public:
 	void RestoreData();
 
 	/** Returns Detail description of what's going with viewport **/
-	FString GetDisplayString() const;
+	FText GetDisplayString() const;
 
 	/** Can we use gizmos? */
 	bool CanUseGizmos() const;
@@ -174,14 +174,14 @@ public:
 	float GetWindStrengthSliderValue() const;
 
 	/** Function to get slider value which returns a string*/
-	FString GetWindStrengthLabel() const;
+	FText GetWindStrengthLabel() const;
 
 	bool IsApplyingClothWind() const;
 
 	/** Show gravity scale */
 	void SetGravityScale( float SliderPos );
 	float GetGravityScaleSliderValue() const;
-	FString GetGravityScaleLabel() const;
+	FText GetGravityScaleLabel() const;
 
 protected:
 
diff --git a/Engine/Source/Editor/Persona/Private/SAnimationSequenceBrowser.cpp b/Engine/Source/Editor/Persona/Private/SAnimationSequenceBrowser.cpp
--- a/Engine/Source/Editor/Persona/Private/SAnimationSequenceBrowser.cpp
+++ b/Engine/Source/Editor/Persona/Private/SAnimationSequenceBrowser.cpp
@@ -763,7 +763,7 @@ TSharedRef<SToolTip> SAnimationSequenceBrowser::CreateCustomAssetToolTip(FAssetD
 				.AutoWidth()
 				[
 					SNew(STextBlock)
-					.Text(TagPair.Value)
+					.Text(FText::FromString(TagPair.Value))
 					.ColorAndOpacity(FSlateColor::UseForeground())
 				]
 			];
@@ -814,7 +814,7 @@ TSharedRef<SToolTip> SAnimationSequenceBrowser::CreateCustomAssetToolTip(FAssetD
 						.HAlign(HAlign_Center)
 						[
 							SNew(STextBlock)
-							.Text(FString(TEXT("No Preview Mesh")))
+							.Text(LOCTEXT("NoPreviewMesh", "No Preview Mesh"))
 						]
 
 						+ SOverlay::Slot()
diff --git a/Engine/Source/Editor/Persona/Private/SMorphTargetViewer.cpp b/Engine/Source/Editor/Persona/Private/SMorphTargetViewer.cpp
--- a/Engine/Source/Editor/Persona/Private/SMorphTargetViewer.cpp
+++ b/Engine/Source/Editor/Persona/Private/SMorphTargetViewer.cpp
@@ -104,7 +104,7 @@ TSharedRef< SWidget > SMorphTargetListRow::GenerateWidgetForColumn( const FName&
 			.VAlign( VAlign_Center )
 			[
 				SNew( STextBlock )
-				.Text( Item->Name.ToString() )
+				.Text( FText::FromName(Item->Name) )
 				.HighlightText( MorphTargetViewer->GetFilterText() )
 			];
 	}
@@ -145,7 +145,7 @@ TSharedRef< SWidget > SMorphTargetListRow::GenerateWidgetForColumn( const FName&
 					.HAlign(HAlign_Right)
 					[
 						SNew(STextBlock)
-						.Text(FString::FromInt(Item->NumberOfVerts))
+						.Text(FText::AsNumber(Item->NumberOfVerts))
 						.HighlightText(MorphTargetViewer->GetFilterText())
 					]
 				];
diff --git a/Engine/Source/Editor/Persona/Private/SRetargetSourceWindow.cpp b/Engine/Source/Editor/Persona/Private/SRetargetSourceWindow.cpp
--- a/Engine/Source/Editor/Persona/Private/SRetargetSourceWindow.cpp
+++ b/Engine/Source/Editor/Persona/Private/SRetargetSourceWindow.cpp
@@ -142,7 +142,7 @@ TSharedRef< SWidget > SRetargetSourceListRow::GenerateWidgetForColumn( const FNa
 			.VAlign( VAlign_Center )
 			[
 				SNew( STextBlock )
-				.Text( Item->GetReferenceMeshName() )
+				.Text( FText::FromString(Item->GetReferenceMeshName()) )
 				.HighlightText( RetargetSourceWindow->GetFilterText() )
 			];
 	}
@@ -173,7 +173,7 @@ void SRetargetSourceWindow::Construct(const FArguments& InArgs)
 		PersonaPtr.Pin()->RegisterOnPostUndo(FPersona::FOnPostUndo::CreateSP( this, &SRetargetSourceWindow::PostUndo ) );
 	}
 
-	FString SkeletonName = Skeleton ? Skeleton->GetName() : LOCTEXT( "RetargetSourceMeshNameLabel", "No Skeleton Present" ).ToString();
+	FText SkeletonName = Skeleton ? FText::FromString(Skeleton->GetName()) : LOCTEXT( "RetargetSourceMeshNameLabel", "No Skeleton Present" );
 
 	ChildSlot
 	[
diff --git a/Engine/Source/Editor/Persona/Private/SSkeletonTree.cpp b/Engine/Source/Editor/Persona/Private/SSkeletonTree.cpp
--- a/Engine/Source/Editor/Persona/Private/SSkeletonTree.cpp
+++ b/Engine/Source/Editor/Persona/Private/SSkeletonTree.cpp
@@ -300,7 +300,7 @@ void FDisplayedMeshBoneInfo::GenerateWidgetForNameColumn( TSharedPtr< SHorizonta
 	const FSlateFontInfo TextFont = GetBoneTextFont( PreviewComponent );
 	const FLinearColor TextColor = GetBoneTextColor( PreviewComponent );
 
-	FString ToolTip = GetBoneToolTip();
+	FText ToolTip = GetBoneToolTip();
 
 	Box->AddSlot()
 		.AutoWidth()
@@ -354,16 +354,16 @@ TSharedRef< SWidget > FDisplayedMeshBoneInfo::CreateBoneTranslationRetargetingMo
 	return MenuBuilder.MakeWidget();
 }
 
-FString FDisplayedMeshBoneInfo::GetTranslationRetargetingModeMenuTitle() const
+FText FDisplayedMeshBoneInfo::GetTranslationRetargetingModeMenuTitle() const
 {
 	const int32 BoneIndex = TargetSkeleton->GetReferenceSkeleton().FindBoneIndex( BoneName );
 	if( BoneIndex != INDEX_NONE )
 	{
 		const EBoneTranslationRetargetingMode::Type RetargetingMode = TargetSkeleton->GetBoneTranslationRetargetingMode(BoneIndex);
-		return TargetSkeleton->GetRetargetingModeString(RetargetingMode);
+		return FText::FromString(TargetSkeleton->GetRetargetingModeString(RetargetingMode));
 	}
 
-	return FString(TEXT("None"));
+	return LOCTEXT("None", "None");
 }
 
 void FDisplayedMeshBoneInfo::SetBoneTranslationRetargetingMode(EBoneTranslationRetargetingMode::Type NewRetargetingMode)
@@ -423,7 +423,7 @@ FReply FDisplayedMeshBoneInfo::OnDragDetected(const FGeometry& MyGeometry, const
 	return FReply::Unhandled();
 }
 
-FString FDisplayedMeshBoneInfo::GetBoneToolTip()
+FText FDisplayedMeshBoneInfo::GetBoneToolTip()
 {
 	bool bIsMeshBone = false;
 	bool bIsWeightedBone = false;
@@ -473,7 +473,7 @@ FString FDisplayedMeshBoneInfo::GetBoneToolTip()
 		}
 	}
 
-	return ToolTip.ToString();
+	return ToolTip;
 }
 
 //////////////////////////////////////////////////////////////////////////
@@ -518,7 +518,7 @@ void FDisplayedSocketInfo::GenerateWidgetForNameColumn( TSharedPtr< SHorizontalB
 		TextColor = FLinearColor::White;
 	}
 
-	FString ToolTip = GetSocketToolTip();
+	FText ToolTip = GetSocketToolTip();
 
 	TAttribute<FText> SocketNameAttr = TAttribute<FText>::Create(TAttribute<FText>::FGetter::CreateSP(this, &FDisplayedSocketInfo::GetSocketNameAsText));
 	TSharedPtr< SInlineEditableTextBlock > InlineWidget;
@@ -645,7 +645,7 @@ FReply FDisplayedSocketInfo::OnDragDetected(const FGeometry& MyGeometry, const F
 	return FReply::Unhandled();
 }
 
-FString FDisplayedSocketInfo::GetSocketToolTip()
+FText FDisplayedSocketInfo::GetSocketToolTip()
 {
 	FText ToolTip;
 
@@ -666,7 +666,7 @@ FString FDisplayedSocketInfo::GetSocketToolTip()
 		ToolTip = LOCTEXT( "SocketToolTipCustomized", "This socket is on the current mesh, customizing the socket of the same name on the skeleton" );
 	}
 
-	return ToolTip.ToString();
+	return ToolTip;
 }
 
 //////////////////////////////////////////////////////////////////////////
@@ -704,7 +704,7 @@ void FDisplayedAttachedAssetInfo::GenerateWidgetForNameColumn( TSharedPtr< SHori
 		[
 			SNew( STextBlock )
 			.ColorAndOpacity(TextColor)
-			.Text( Asset->GetName() )
+			.Text( FText::FromString(Asset->GetName()) )
 			.HighlightText( FilterText )
 			.Font(TextFont)
 		];
@@ -877,7 +877,7 @@ void SSkeletonTree::Construct(const FArguments& InArgs)
 				.OnCheckStateChanged(this, &SSkeletonTree::OnChangeShowingRetargetingOptions)
 				[
 					SNew(STextBlock)
-						.Text(LOCTEXT("ShowRetargetingOptions", "Show Retargeting Options").ToString())
+						.Text(LOCTEXT("ShowRetargetingOptions", "Show Retargeting Options"))
 				]
 			]
 		]
@@ -2284,7 +2284,7 @@ bool SSkeletonTree::IsBoneWeighted( int32 MeshBoneIndex, UDebugSkelMeshComponent
 	return Index != INDEX_NONE;
 }
 
-FString SSkeletonTree::GetBoneFilterMenuTitle() const
+FText SSkeletonTree::GetBoneFilterMenuTitle() const
 {
 	FText BoneFilterMenuText;
 
@@ -2312,10 +2312,10 @@ FString SSkeletonTree::GetBoneFilterMenuTitle() const
 		break;
 	}
 
-	return BoneFilterMenuText.ToString();
+	return BoneFilterMenuText;
 }
 
-FString SSkeletonTree::GetSocketFilterMenuTitle() const
+FText SSkeletonTree::GetSocketFilterMenuTitle() const
 {
 	FText SocketFilterMenuText;
 
@@ -2347,7 +2347,7 @@ FString SSkeletonTree::GetSocketFilterMenuTitle() const
 		break;
 	}
 
-	return SocketFilterMenuText.ToString();
+	return SocketFilterMenuText;
 }
 
 void SSkeletonTree::OnPreviewMeshChanged( USkeletalMesh* NewPreviewMesh )
diff --git a/Engine/Source/Editor/Persona/Private/SSkeletonTree.h b/Engine/Source/Editor/Persona/Private/SSkeletonTree.h
--- a/Engine/Source/Editor/Persona/Private/SSkeletonTree.h
+++ b/Engine/Source/Editor/Persona/Private/SSkeletonTree.h
@@ -173,7 +173,7 @@ private:
 	FLinearColor GetBoneTextColor( UDebugSkelMeshComponent* PreviewComponent ) const;
 
 	/** Function that returns the current tooltip for this bone, depending on how it's used by the mesh */
-	FString GetBoneToolTip();
+	FText GetBoneToolTip();
 
 	/** The actual bone data that we create Slate widgets to display */
 	FName BoneName;
@@ -182,7 +182,7 @@ private:
 	TSharedRef< SWidget > CreateBoneTranslationRetargetingModeMenu();
 
 	/** Get Title for Bone Translation Retargeting Mode menu. */
-	FString GetTranslationRetargetingModeMenuTitle() const;
+	FText GetTranslationRetargetingModeMenuTitle() const;
 };
 
 //////////////////////////////////////////////////////////////////////////
@@ -267,7 +267,7 @@ protected:
 	void OnCommitSocketName( const FText& InText, ETextCommit::Type CommitInfo );
 
 	/** Function that returns the current tooltip for this socket */
-	FString GetSocketToolTip();
+	FText GetSocketToolTip();
 
 	/** Pointer to the socket */
 	USkeletalMeshSocket*	SocketData;
@@ -526,10 +526,10 @@ private:
 	bool IsSocketFilter( ESocketFilter::Type InSocketFilter ) const;
 
 	/** Returns the current text for the bone filter button - "All", "Mesh" or "Weighted" */
-	FString GetBoneFilterMenuTitle() const;
+	FText GetBoneFilterMenuTitle() const;
 
 	/** Returns the current text for the socket filter button - "All", "Mesh" or "Skeleton" */
-	FString GetSocketFilterMenuTitle() const;
+	FText GetSocketFilterMenuTitle() const;
 
 	/** We can only add sockets in Active, Skeleton or All mode (otherwise they just disappear) */
 	bool IsAddingSocketsAllowed() const;
diff --git a/Engine/Source/Editor/Persona/Private/Shared/BoneSelectionWidget.cpp b/Engine/Source/Editor/Persona/Private/Shared/BoneSelectionWidget.cpp
--- a/Engine/Source/Editor/Persona/Private/Shared/BoneSelectionWidget.cpp
+++ b/Engine/Source/Editor/Persona/Private/Shared/BoneSelectionWidget.cpp
@@ -102,7 +102,7 @@ TSharedRef<ITableRow> SBoneSelectionWidget::MakeTreeRowWidget(TSharedPtr<FBoneNa
 		[
 			SNew(STextBlock)
 			.HighlightText(FilterText)
-				.Text(InInfo->BoneName.ToString())
+			.Text(FText::FromName(InInfo->BoneName))
 		];
 }
 
@@ -190,17 +190,17 @@ void SBoneSelectionWidget::OnSelectionChanged(TSharedPtr<FBoneNameInfo> BoneInfo
 	BonePickerButton->SetIsOpen(false);
 }
 
-FString SBoneSelectionWidget::GetCurrentBoneName() const
+FText SBoneSelectionWidget::GetCurrentBoneName() const
 {
 	if(OnGetSelectedBone.IsBound())
 	{
 		FName Name = OnGetSelectedBone.Execute();
 
-		return Name.ToString();
+		return FText::FromName(Name);
 	}
 
 	// @todo implement default solution?
-	return TEXT("");
+	return FText::GetEmpty();
 }
 #undef LOCTEXT_NAMESPACE
 
diff --git a/Engine/Source/Editor/Persona/Public/BoneDragDropOp.h b/Engine/Source/Editor/Persona/Public/BoneDragDropOp.h
--- a/Engine/Source/Editor/Persona/Public/BoneDragDropOp.h
+++ b/Engine/Source/Editor/Persona/Public/BoneDragDropOp.h
@@ -35,10 +35,9 @@ public:
 			];
 	}
 
-	FString GetHoverText() const
+	FText GetHoverText() const
 	{
-		FString HoverText = FString::Printf(TEXT("Bone %s"), *BoneName.GetPlainNameString());
-		return HoverText;
+		return FText::Format(NSLOCTEXT("BoneDragDropOp", "BoneHoverTextFmt", "Bone {0}"), FText::FromString(BoneName.GetPlainNameString()));
 	}
 
 	const FSlateBrush* GetIcon( ) const
diff --git a/Engine/Source/Editor/Persona/Public/BoneSelectionWidget.h b/Engine/Source/Editor/Persona/Public/BoneSelectionWidget.h
--- a/Engine/Source/Editor/Persona/Public/BoneSelectionWidget.h
+++ b/Engine/Source/Editor/Persona/Public/BoneSelectionWidget.h
@@ -63,7 +63,7 @@ private:
 	// Called when the user selects a bone name
 	void OnSelectionChanged(TSharedPtr<FBoneNameInfo>, ESelectInfo::Type SelectInfo);
 	// Gets the current bone name, used to get the right name for the combo button
-	FString GetCurrentBoneName() const;
+	FText GetCurrentBoneName() const;
 
 	// Tree view used in the button menu
 	TSharedPtr<STreeView<TSharedPtr<FBoneNameInfo>>> TreeView;
diff --git a/Engine/Source/Editor/Persona/Public/SocketDragDropOp.h b/Engine/Source/Editor/Persona/Public/SocketDragDropOp.h
--- a/Engine/Source/Editor/Persona/Public/SocketDragDropOp.h
+++ b/Engine/Source/Editor/Persona/Public/SocketDragDropOp.h
@@ -36,10 +36,9 @@ public:
 	}
 
 	/** Passed into STextBlock so Slate can grab the current text for display */
-	FString GetHoverText() const
+	FText GetHoverText() const
 	{
-		FString HoverText = FString::Printf(TEXT("Socket %s"), *SocketInfo.Socket->SocketName.ToString() );
-		return HoverText;
+		return FText::Format(NSLOCTEXT("SocketDragDrop", "SocketNameFmt", "Socket {0}"), FText::FromName(SocketInfo.Socket->SocketName));
 	}
 
 	/** Passed into SImage so Slate can grab the current icon for display */
diff --git a/Engine/Source/Editor/PhAT/Private/PhAT.cpp b/Engine/Source/Editor/PhAT/Private/PhAT.cpp
--- a/Engine/Source/Editor/PhAT/Private/PhAT.cpp
+++ b/Engine/Source/Editor/PhAT/Private/PhAT.cpp
@@ -1435,7 +1435,7 @@ TSharedRef<ITableRow> FPhAT::OnGenerateRowForTree(FTreeElemPtr Item, const TShar
 	[
 		SNew(STextBlock)
 		.Font((*Item).bBold? FEditorStyle::GetFontStyle("BoldFont"): FEditorStyle::GetFontStyle("NormalFont"))
-		.Text((*Item).Name.ToString())
+		.Text(FText::FromName((*Item).Name))
 	];
 }
 
@@ -3063,9 +3063,8 @@ TSharedRef<SWidget> FPhAT::BuildHierarchyFilterMenu()
 	return MenuBuilder.MakeWidget();
 }
 
-FString FPhAT::GetHierarchyFilter() const
+FText FPhAT::GetHierarchyFilter() const
 {
-
 	FText FilterMenuText;
 
 	switch (HierarchyFilterMode)
@@ -3082,7 +3081,7 @@ FString FPhAT::GetHierarchyFilter() const
 		break;
 	}
 
-	return FilterMenuText.ToString();
+	return FilterMenuText;
 
 }
 
diff --git a/Engine/Source/Editor/PhAT/Private/PhAT.h b/Engine/Source/Editor/PhAT/Private/PhAT.h
--- a/Engine/Source/Editor/PhAT/Private/PhAT.h
+++ b/Engine/Source/Editor/PhAT/Private/PhAT.h
@@ -68,7 +68,7 @@ public:
 	TSharedPtr<SWidget> BuildMenuWidgetBone();
 	TSharedRef<SWidget> BuildStaticMeshAssetPicker();
 	TSharedRef<SWidget> BuildHierarchyFilterMenu();
-	FString GetHierarchyFilter() const;
+	FText GetHierarchyFilter() const;
 
 	/** IToolkit interface */
 	virtual FName GetToolkitFName() const override;
diff --git a/Engine/Source/Editor/PropertyEditor/Private/SDetailNameArea.cpp b/Engine/Source/Editor/PropertyEditor/Private/SDetailNameArea.cpp
--- a/Engine/Source/Editor/PropertyEditor/Private/SDetailNameArea.cpp
+++ b/Engine/Source/Editor/PropertyEditor/Private/SDetailNameArea.cpp
@@ -190,7 +190,7 @@ void SDetailNameArea::BuildObjectNameAreaSelectionLabel( TSharedRef< SHorizontal
 	}
 	else
 	{
-		const FString SelectionText = FString::Printf( *LOCTEXT("MultipleObjectsSelected", "%d objects").ToString(), NumSelectedObjects );
+		const FText SelectionText = FText::Format( LOCTEXT("MultipleObjectsSelectedFmt", "{0} objects"), FText::AsNumber(NumSelectedObjects) );
 		SelectionLabelBox->AddSlot()
 		.VAlign(VAlign_Center)
 		.HAlign( HAlign_Left )
diff --git a/Engine/Source/Editor/PropertyEditor/Private/SSingleProperty.cpp b/Engine/Source/Editor/PropertyEditor/Private/SSingleProperty.cpp
--- a/Engine/Source/Editor/PropertyEditor/Private/SSingleProperty.cpp
+++ b/Engine/Source/Editor/PropertyEditor/Private/SSingleProperty.cpp
@@ -184,8 +184,8 @@ void SSingleProperty::SetObject( UObject* InObject )
 		[
 			SNew(STextBlock)
 			.Font(PropertyFont)
-			.Text(NSLOCTEXT("PropertyEditor", "SinglePropertyInvalidType", "Cannot Edit Inline").ToString())
-			.ToolTipText(NSLOCTEXT("PropertyEditor", "SinglePropertyInvalidType_Tooltip", "Properties of this type cannot be edited inline; edit it elsewhere").ToString())
+			.Text(NSLOCTEXT("PropertyEditor", "SinglePropertyInvalidType", "Cannot Edit Inline"))
+			.ToolTipText(NSLOCTEXT("PropertyEditor", "SinglePropertyInvalidType_Tooltip", "Properties of this type cannot be edited inline; edit it elsewhere"))
 		];
 
 		// invalid or missing property
diff --git a/Engine/Source/Editor/PropertyEditor/Private/UserInterface/PropertyEditor/SPropertyComboBox.cpp b/Engine/Source/Editor/PropertyEditor/Private/UserInterface/PropertyEditor/SPropertyComboBox.cpp
--- a/Engine/Source/Editor/PropertyEditor/Private/UserInterface/PropertyEditor/SPropertyComboBox.cpp
+++ b/Engine/Source/Editor/PropertyEditor/Private/UserInterface/PropertyEditor/SPropertyComboBox.cpp
@@ -26,11 +26,12 @@ void SPropertyComboBox::Construct( const FArguments& InArgs )
 		}
 	}
 
+	auto VisibleTextAttr = InArgs._VisibleText;
 	SComboBox< TSharedPtr<FString> >::Construct(SComboBox< TSharedPtr<FString> >::FArguments()
 		.Content()
 		[
 			SNew( STextBlock )
-			.Text( InArgs._VisibleText )
+			.Text_Lambda( [=] { return (VisibleTextAttr.IsSet()) ? FText::FromString(VisibleTextAttr.Get()) : FText::GetEmpty(); } )
 			.Font( Font )
 		]
 		.OptionsSource(&ComboItemList)
@@ -120,7 +121,7 @@ TSharedRef<SWidget> SPropertyComboBox::OnGenerateComboWidget( TSharedPtr<FString
 
 	return
 		SNew( STextBlock )
-		.Text( *InComboString )
+		.Text( FText::FromString(*InComboString) )
 		.Font( Font )
 		.ToolTipText(ToolTip)
 		.IsEnabled(bEnabled);
diff --git a/Engine/Source/Editor/PropertyEditor/Private/UserInterface/PropertyEditor/SPropertyEditorArray.h b/Engine/Source/Editor/PropertyEditor/Private/UserInterface/PropertyEditor/SPropertyEditorArray.h
--- a/Engine/Source/Editor/PropertyEditor/Private/UserInterface/PropertyEditor/SPropertyEditorArray.h
+++ b/Engine/Source/Editor/PropertyEditor/Private/UserInterface/PropertyEditor/SPropertyEditorArray.h
@@ -21,7 +21,7 @@ public:
 	{
 		PropertyEditor = InPropertyEditor;
 
-		TAttribute<FString> TextAttr;
+		TAttribute<FText> TextAttr;
 		if( PropertyEditorHelpers::IsStaticArray( *InPropertyEditor->GetPropertyNode() ) )
 		{
 			// Static arrays need special case handling for their values
@@ -57,9 +57,9 @@ public:
 		OutMaxDesiredWidth = 130.0f;
 	}
 private:
-	FString GetArrayTextValue() const
+	FText GetArrayTextValue() const
 	{
-		return FString::Printf( *LOCTEXT("NumArrayItems", "%d elements").ToString() , PropertyEditor->GetPropertyNode()->GetNumChildNodes() );
+		return FText::Format( LOCTEXT("NumArrayItemsFmt", "{0} elements"), FText::AsNumber(PropertyEditor->GetPropertyNode()->GetNumChildNodes()) );
 	}
 
 	/** @return True if the property can be edited */
diff --git a/Engine/Source/Editor/PropertyEditor/Private/UserInterface/PropertyEditor/SPropertyEditorArrayItem.cpp b/Engine/Source/Editor/PropertyEditor/Private/UserInterface/PropertyEditor/SPropertyEditorArrayItem.cpp
--- a/Engine/Source/Editor/PropertyEditor/Private/UserInterface/PropertyEditor/SPropertyEditorArrayItem.cpp
+++ b/Engine/Source/Editor/PropertyEditor/Private/UserInterface/PropertyEditor/SPropertyEditorArrayItem.cpp
@@ -37,15 +37,15 @@ bool SPropertyEditorArrayItem::Supports( const TSharedRef< class FPropertyEditor
 		!(Cast<const UArrayProperty>(Property->GetOuter())->PropertyFlags & CPF_EditConst);
 }
 
-FString SPropertyEditorArrayItem::GetValueAsString() const
+FText SPropertyEditorArrayItem::GetValueAsString() const
 {
 	if( PropertyEditor->GetProperty() && PropertyEditor->PropertyIsA( UStructProperty::StaticClass() ) )
 	{
-		return FText::Format( NSLOCTEXT("PropertyEditor", "NumStructItems", "{0} members"), FText::AsNumber( PropertyEditor->GetPropertyNode()->GetNumChildNodes() ) ).ToString();
+		return FText::Format( NSLOCTEXT("PropertyEditor", "NumStructItems", "{0} members"), FText::AsNumber( PropertyEditor->GetPropertyNode()->GetNumChildNodes() ) );
 	}
 	else
 	{
-		return PropertyEditor->GetValueAsString();
+		return FText::FromString(PropertyEditor->GetValueAsString());
 	}
 }
 
diff --git a/Engine/Source/Editor/PropertyEditor/Private/UserInterface/PropertyEditor/SPropertyEditorArrayItem.h b/Engine/Source/Editor/PropertyEditor/Private/UserInterface/PropertyEditor/SPropertyEditorArrayItem.h
--- a/Engine/Source/Editor/PropertyEditor/Private/UserInterface/PropertyEditor/SPropertyEditorArrayItem.h
+++ b/Engine/Source/Editor/PropertyEditor/Private/UserInterface/PropertyEditor/SPropertyEditorArrayItem.h
@@ -19,7 +19,7 @@ public:
 	void GetDesiredWidth( float& OutMinDesiredWidth, float& OutMaxDesiredWidth );
 
 private:
-	FString GetValueAsString() const;
+	FText GetValueAsString() const;
 
 	/** @return True if the property can be edited */
 	bool CanEdit() const;
diff --git a/Engine/Source/Editor/PropertyEditor/Private/UserInterface/PropertyEditor/SPropertyEditorAsset.cpp b/Engine/Source/Editor/PropertyEditor/Private/UserInterface/PropertyEditor/SPropertyEditorAsset.cpp
--- a/Engine/Source/Editor/PropertyEditor/Private/UserInterface/PropertyEditor/SPropertyEditorAsset.cpp
+++ b/Engine/Source/Editor/PropertyEditor/Private/UserInterface/PropertyEditor/SPropertyEditorAsset.cpp
@@ -355,12 +355,12 @@ void SPropertyEditorAsset::CloseComboButton()
 	AssetComboButton->SetIsOpen(false);
 }
 
-FString SPropertyEditorAsset::OnGetAssetName() const
+FText SPropertyEditorAsset::OnGetAssetName() const
 {
 	FObjectOrAssetData Value; 
 	FPropertyAccess::Result Result = GetValue( Value );
 
-	FString Name = LOCTEXT("None", "None").ToString();
+	FText Name = LOCTEXT("None", "None");
 	if( Result == FPropertyAccess::Success )
 	{
 		if(Value.Object != NULL)
@@ -368,64 +368,63 @@ FString SPropertyEditorAsset::OnGetAssetName() const
 			if( bIsActor )
 			{
 				AActor* Actor = CastChecked<AActor>(Value.Object);
-				Name = Actor->GetActorLabel();
+				Name = FText::FromString(Actor->GetActorLabel());
 			}
 			else
 			{
-				Name = Value.Object->GetName();
+				Name = FText::FromString(Value.Object->GetName());
 			}
 		}
 		else if( Value.AssetData.IsValid() )
 		{
-			Name = FString( Value.AssetData.AssetName.ToString() );
+			Name = FText::FromName(Value.AssetData.AssetName);
 		}
 	}
 	else if( Result == FPropertyAccess::MultipleValues )
 	{
-		Name = LOCTEXT("MultipleValues", "Multiple Values").ToString();
+		Name = LOCTEXT("MultipleValues", "Multiple Values");
 	}
 
 	return Name;
 }
 
-FString SPropertyEditorAsset::OnGetAssetClassName() const
+FText SPropertyEditorAsset::OnGetAssetClassName() const
 {
 	UClass* Class = GetDisplayedClass();
 	if(Class)
 	{
-		return Class->GetName();
+		return FText::FromString(Class->GetName());
 	}
-	return FString();
+	return FText::GetEmpty();
 }
 
-FString SPropertyEditorAsset::OnGetToolTip() const
+FText SPropertyEditorAsset::OnGetToolTip() const
 {
 	FObjectOrAssetData Value; 
 	FPropertyAccess::Result Result = GetValue( Value );
 
-	FString ToolTip;
+	FText ToolTip = FText::GetEmpty();
 
 	if( Result == FPropertyAccess::Success )
 	{
 		if(Value.Object != NULL && !bIsActor )
 		{
 			// Display the package name which is a valid path to the object without redundant information
-			ToolTip = Value.Object->GetOutermost()->GetName();
+			ToolTip = FText::FromString(Value.Object->GetOutermost()->GetName());
 		}
 		else if( Value.AssetData.IsValid() )
 		{
-			ToolTip = Value.AssetData.PackageName.ToString();
+			ToolTip = FText::FromName(Value.AssetData.PackageName);
 		}
 	}
 	else if( Result == FPropertyAccess::MultipleValues )
 	{
-		ToolTip = LOCTEXT("MultipleValues", "Multiple Values").ToString();
+		ToolTip = LOCTEXT("MultipleValues", "Multiple Values");
 	}
 
-
 	if( ToolTip.IsEmpty() )
 	{
-		ToolTip = ObjectPath.Get();
+		ToolTip = FText::FromString(ObjectPath.Get());
 	}
 
 	return ToolTip;
diff --git a/Engine/Source/Editor/PropertyEditor/Private/UserInterface/PropertyEditor/SPropertyEditorAsset.h b/Engine/Source/Editor/PropertyEditor/Private/UserInterface/PropertyEditor/SPropertyEditorAsset.h
--- a/Engine/Source/Editor/PropertyEditor/Private/UserInterface/PropertyEditor/SPropertyEditorAsset.h
+++ b/Engine/Source/Editor/PropertyEditor/Private/UserInterface/PropertyEditor/SPropertyEditorAsset.h
@@ -116,19 +116,19 @@ private:
 	 * Get the name to be displayed for this asset
 	 * @returns the name of the asset
 	 */
-	FString OnGetAssetName() const;
+	FText OnGetAssetName() const;
 
 	/** 
 	 * Get the class name to be displayed for this asset
 	 * @returns the class name of this asset
 	 */
-	FString OnGetAssetClassName() const;
+	FText OnGetAssetClassName() const;
 
 	/** 
 	 * Get the tooltip to be displayed for this asset
 	 * @returns the tooltip string
 	 */
-	FString OnGetToolTip() const;
+	FText OnGetToolTip() const;
 
 	/** 
 	 * Set the value of the asset referenced by this property editor.
diff --git a/Engine/Source/Editor/PropertyEditor/Private/UserInterface/PropertyEditor/SPropertyEditorClass.cpp b/Engine/Source/Editor/PropertyEditor/Private/UserInterface/PropertyEditor/SPropertyEditorClass.cpp
--- a/Engine/Source/Editor/PropertyEditor/Private/UserInterface/PropertyEditor/SPropertyEditorClass.cpp
+++ b/Engine/Source/Editor/PropertyEditor/Private/UserInterface/PropertyEditor/SPropertyEditorClass.cpp
@@ -158,7 +158,7 @@ static FString GetClassDisplayName(const UObject* Object)
 	return (Object) ? Object->GetName() : "None";
 }
 
-FString SPropertyEditorClass::GetDisplayValueAsString() const
+FText SPropertyEditorClass::GetDisplayValueAsString() const
 {
 	if(PropertyEditor.IsValid())
 	{
@@ -167,13 +167,13 @@ FString SPropertyEditorClass::GetDisplayValueAsString() const
 
 		if(Result == FPropertyAccess::Success && ObjectValue != NULL)
 		{
-			return GetClassDisplayName(ObjectValue);
+			return FText::FromString(GetClassDisplayName(ObjectValue));
 		}
 
-		return FPaths::GetBaseFilename(PropertyEditor->GetValueAsString());
+		return FText::FromString(FPaths::GetBaseFilename(PropertyEditor->GetValueAsString()));
 	}
 
-	return GetClassDisplayName(SelectedClass.Get());
+	return FText::FromString(GetClassDisplayName(SelectedClass.Get()));
 }
 
 TSharedRef<SWidget> SPropertyEditorClass::GenerateClassPicker()
diff --git a/Engine/Source/Editor/PropertyEditor/Private/UserInterface/PropertyEditor/SPropertyEditorClass.h b/Engine/Source/Editor/PropertyEditor/Private/UserInterface/PropertyEditor/SPropertyEditorClass.h
--- a/Engine/Source/Editor/PropertyEditor/Private/UserInterface/PropertyEditor/SPropertyEditorClass.h
+++ b/Engine/Source/Editor/PropertyEditor/Private/UserInterface/PropertyEditor/SPropertyEditorClass.h
@@ -66,7 +66,7 @@ private:
 	/**
 	 * Gets the active display value as a string
 	 */
-	FString GetDisplayValueAsString() const;
+	FText GetDisplayValueAsString() const;
 
 private:
 	/** The property editor we were constructed for, or null if we're editing using the construction arguments */
diff --git a/Engine/Source/Editor/PropertyEditor/Private/UserInterface/PropertyEditor/SPropertySceneOutliner.cpp b/Engine/Source/Editor/PropertyEditor/Private/UserInterface/PropertyEditor/SPropertySceneOutliner.cpp
--- a/Engine/Source/Editor/PropertyEditor/Private/UserInterface/PropertyEditor/SPropertySceneOutliner.cpp
+++ b/Engine/Source/Editor/PropertyEditor/Private/UserInterface/PropertyEditor/SPropertySceneOutliner.cpp
@@ -30,7 +30,7 @@ void SPropertySceneOutliner::Construct( const FArguments& InArgs )
 			SNew( SButton )
 			.ButtonStyle( FEditorStyle::Get(), "HoverHintOnly" )
 			.OnClicked( this, &SPropertySceneOutliner::OnClicked )
-			.ToolTipText(LOCTEXT("PickButtonLabel", "Pick Actor").ToString())
+			.ToolTipText(LOCTEXT("PickButtonLabel", "Pick Actor"))
 			.ContentPadding(0)
 			.ForegroundColor( FSlateColor::UseForeground() )
 			.IsFocusable(false)
diff --git a/Engine/Source/Editor/PropertyEditor/Private/UserInterface/PropertyTable/ObjectNameTableCellPresenter.h b/Engine/Source/Editor/PropertyEditor/Private/UserInterface/PropertyTable/ObjectNameTableCellPresenter.h
--- a/Engine/Source/Editor/PropertyEditor/Private/UserInterface/PropertyTable/ObjectNameTableCellPresenter.h
+++ b/Engine/Source/Editor/PropertyEditor/Private/UserInterface/PropertyTable/ObjectNameTableCellPresenter.h
@@ -81,7 +81,7 @@ private:
 				[
 					SNew( STextBlock )
 					.Font( FEditorStyle::GetFontStyle( TextFontStyle ) )
-					.Text( DisplayNamePieces[ Index ] )
+					.Text( FText::FromString(DisplayNamePieces[ Index ]) )
 				];
 
 			if ( Index < DisplayNamePieces.Num() - 1 )
diff --git a/Engine/Source/Editor/PropertyEditor/Private/UserInterface/PropertyTable/SColumnHeader.h b/Engine/Source/Editor/PropertyEditor/Private/UserInterface/PropertyTable/SColumnHeader.h
--- a/Engine/Source/Editor/PropertyEditor/Private/UserInterface/PropertyTable/SColumnHeader.h
+++ b/Engine/Source/Editor/PropertyEditor/Private/UserInterface/PropertyTable/SColumnHeader.h
@@ -99,7 +99,7 @@ protected:
 			[
 				SNew( STextBlock )
 				.Font( FEditorStyle::GetFontStyle( TextFontStyle ) )
-				.Text( DisplayNamePieces[ Index ] )
+				.Text( FText::FromString(DisplayNamePieces[ Index ]) )
 			];
 
 			if ( Index < DisplayNamePieces.Num() - 1 )
diff --git a/Engine/Source/Editor/PropertyEditor/Private/UserInterface/PropertyTable/SPropertyNameColumnHeader.h b/Engine/Source/Editor/PropertyEditor/Private/UserInterface/PropertyTable/SPropertyNameColumnHeader.h
--- a/Engine/Source/Editor/PropertyEditor/Private/UserInterface/PropertyTable/SPropertyNameColumnHeader.h
+++ b/Engine/Source/Editor/PropertyEditor/Private/UserInterface/PropertyTable/SPropertyNameColumnHeader.h
@@ -61,7 +61,7 @@ class SPropertyNameColumnHeader : public SColumnHeader
 				.Padding( FMargin( 2, 0, 2, 0 ) )
 				[
 					SNew( STextBlock )
-					.Text( PropertyName )
+					.Text( FText::FromString(PropertyName) )
 					.ColorAndOpacity(FLinearColor(1.0f, 1.0f, 1.0f))
 				]
 			];
diff --git a/Engine/Source/Editor/PropertyEditor/Private/UserInterface/PropertyTable/SPropertyTableCell.cpp b/Engine/Source/Editor/PropertyEditor/Private/UserInterface/PropertyTable/SPropertyTableCell.cpp
--- a/Engine/Source/Editor/PropertyEditor/Private/UserInterface/PropertyTable/SPropertyTableCell.cpp
+++ b/Engine/Source/Editor/PropertyEditor/Private/UserInterface/PropertyTable/SPropertyTableCell.cpp
@@ -225,7 +225,7 @@ TSharedRef<SBorder> SPropertyTableCell::ConstructInvalidPropertyWidget()
 			[
 				SNew(STextBlock)
 				.ColorAndOpacity(FLinearColor::Red)
-				.Text(NSLOCTEXT("PropertyEditor", "InvalidTableCellProperty", "Failed to retrieve value").ToString())
+				.Text(NSLOCTEXT("PropertyEditor", "InvalidTableCellProperty", "Failed to retrieve value"))
 			]
 		];
 }
diff --git a/Engine/Source/Editor/PropertyEditor/Private/UserInterface/PropertyTable/TextPropertyTableCellPresenter.cpp b/Engine/Source/Editor/PropertyEditor/Private/UserInterface/PropertyTable/TextPropertyTableCellPresenter.cpp
--- a/Engine/Source/Editor/PropertyEditor/Private/UserInterface/PropertyTable/TextPropertyTableCellPresenter.cpp
+++ b/Engine/Source/Editor/PropertyEditor/Private/UserInterface/PropertyTable/TextPropertyTableCellPresenter.cpp
@@ -41,7 +41,7 @@ TSharedRef< class SWidget > FTextPropertyTableCellPresenter::ConstructDisplayWid
 		.Padding( FMargin( 2, 0, 0, 0 ) )
 		[
 			SNew( STextBlock )
-			.Text( PropertyEditor->GetValueAsString() )
+			.Text( PropertyEditor->GetValueAsText() )
 			.ToolTipText( PropertyEditor->GetToolTipText() )
 			.Font( Font )
 		]
diff --git a/Engine/Source/Editor/SceneOutliner/Private/SSceneOutliner.cpp b/Engine/Source/Editor/SceneOutliner/Private/SSceneOutliner.cpp
--- a/Engine/Source/Editor/SceneOutliner/Private/SSceneOutliner.cpp
+++ b/Engine/Source/Editor/SceneOutliner/Private/SSceneOutliner.cpp
@@ -2012,7 +2012,7 @@ namespace SceneOutliner
 		return ( IsFilterActive() || RootTreeItems.Num() > 0 ) ? EVisibility::Collapsed : EVisibility::Visible;
 	}
 
-	FString SSceneOutliner::GetFilterStatusText() const
+	FText SSceneOutliner::GetFilterStatusText() const
 	{
 		const int32 TotalActorCount = ApplicableActors.Num();
 
@@ -2027,24 +2027,24 @@ namespace SceneOutliner
 		{
 			if (SelectedActorCount == 0)
 			{
-				return FString::Printf( *LOCTEXT("ShowingAllActors", "%d actors").ToString(), TotalActorCount );
+				return FText::Format( LOCTEXT("ShowingAllActorsFmt", "{0} actors"), FText::AsNumber( TotalActorCount ) );
 			}
 			else
 			{
-				return FString::Printf( *LOCTEXT("ShowingAllActorsSelected", "%d actors (%d selected)").ToString(), TotalActorCount, SelectedActorCount );
+				return FText::Format( LOCTEXT("ShowingAllActorsSelectedFmt", "{0} actors ({1} selected)"), FText::AsNumber( TotalActorCount ), FText::AsNumber( SelectedActorCount ) );
 			}
 		}
 		else if( IsFilterActive() && FilteredActorCount == 0 )
 		{
-			return FString::Printf( *LOCTEXT("ShowingNoActors", "No matching actors (%d total)").ToString(), TotalActorCount );
+			return FText::Format( LOCTEXT("ShowingNoActorsFmt", "No matching actors ({0} total)"), FText::AsNumber( TotalActorCount ) );
 		}
 		else if (SelectedActorCount != 0)
 		{
-			return FString::Printf( *LOCTEXT("ShowingOnlySomeActorsSelected", "Showing %d of %d actors (%d selected)").ToString(), FilteredActorCount, TotalActorCount, SelectedActorCount );
+			return FText::Format( LOCTEXT("ShowingOnlySomeActorsSelectedFmt", "Showing {0} of {1} actors ({2} selected)"), FText::AsNumber( FilteredActorCount ), FText::AsNumber( TotalActorCount ), FText::AsNumber( SelectedActorCount ) );
 		}
 		else
 		{
-			return FString::Printf( *LOCTEXT("ShowingOnlySomeActors", "Showing %d of %d actors").ToString(), FilteredActorCount, TotalActorCount );
+			return FText::Format( LOCTEXT("ShowingOnlySomeActorsFmt", "Showing {0} of {1} actors"), FText::AsNumber( FilteredActorCount ), FText::AsNumber( TotalActorCount ) );
 		}
 	}
 
