Mesh paint editor mode: Refactored texture override code to get closer to being able to override via the IMeshPaintGeometryAdapter

diff --git a/Engine/Plugins/2D/Paper2D/Source/Paper2D/Private/PaperSpriteComponent.cpp b/Engine/Plugins/2D/Paper2D/Source/Paper2D/Private/PaperSpriteComponent.cpp
--- a/Engine/Plugins/2D/Paper2D/Source/Paper2D/Private/PaperSpriteComponent.cpp
+++ b/Engine/Plugins/2D/Paper2D/Source/Paper2D/Private/PaperSpriteComponent.cpp
@@ -235,6 +235,9 @@ void UPaperSpriteComponent::GetUsedTextures(TArray<UTexture*>& OutTextures, EMat
 		}
 	}
 
+	// Add any extra textures
+	//@TODO: 
+
 	// Get any textures referenced by our materials
 	Super::GetUsedTextures(OutTextures, QualityLevel);
 }
diff --git a/Engine/Plugins/2D/Paper2D/Source/Paper2DEditor/Private/MeshPainting/MeshPaintSpriteAdapter.cpp b/Engine/Plugins/2D/Paper2D/Source/Paper2DEditor/Private/MeshPainting/MeshPaintSpriteAdapter.cpp
--- a/Engine/Plugins/2D/Paper2D/Source/Paper2DEditor/Private/MeshPainting/MeshPaintSpriteAdapter.cpp
+++ b/Engine/Plugins/2D/Paper2D/Source/Paper2DEditor/Private/MeshPainting/MeshPaintSpriteAdapter.cpp
@@ -134,6 +134,23 @@ void FMeshPaintSpriteAdapter::QueryPaintableTextures(int32 MaterialIndex, int32&
 	}
 }
 
+void FMeshPaintSpriteAdapter::ApplyOrRemoveTextureOverride(UTexture* SourceTexture, UTexture* OverrideTexture) const
+{
+	UPaperSprite* Sprite = SpriteComponent->GetSprite();
+	checkSlow(Sprite != nullptr);
+
+	if (Sprite->GetSourceTexture() == SourceTexture)
+	{
+		//@TOOD: Need to set SourceTexture override on the component or something like that!
+	}
+	//else if (Sprite->AdditionalSourceTextures.Contains(SourceTexture)) //@TOOD: Need to swap it here too
+	//{
+	//}
+
+	// Make sure we swap it out on any textures that aren't part of the sprite as well
+	DefaultApplyOrRemoveTextureOverride(SpriteComponent, SourceTexture, OverrideTexture);
+}
+
 //////////////////////////////////////////////////////////////////////////
 // FMeshPaintSpriteAdapterFactory
 
diff --git a/Engine/Plugins/2D/Paper2D/Source/Paper2DEditor/Private/MeshPainting/MeshPaintSpriteAdapter.h b/Engine/Plugins/2D/Paper2D/Source/Paper2DEditor/Private/MeshPainting/MeshPaintSpriteAdapter.h
--- a/Engine/Plugins/2D/Paper2D/Source/Paper2DEditor/Private/MeshPainting/MeshPaintSpriteAdapter.h
+++ b/Engine/Plugins/2D/Paper2D/Source/Paper2DEditor/Private/MeshPainting/MeshPaintSpriteAdapter.h
@@ -18,6 +18,7 @@ public:
 	virtual bool LineTraceComponent(struct FHitResult& OutHit, const FVector Start, const FVector End, const struct FCollisionQueryParams& Params) const override;
 	virtual void SphereIntersectTriangles(TArray<int32>& OutTriangles, const float ComponentSpaceSquaredBrushRadius, const FVector& ComponentSpaceBrushPosition) const override;
 	virtual void QueryPaintableTextures(int32 MaterialIndex, int32& OutDefaultIndex, TArray<struct FPaintableTexture>& InOutTextureList) override;
+	virtual void ApplyOrRemoveTextureOverride(UTexture* SourceTexture, UTexture* OverrideTexture) const override;
 
 protected:
 	UPaperSpriteComponent* SpriteComponent;
diff --git a/Engine/Source/Editor/MeshPaint/Private/MeshPaintEdMode.cpp b/Engine/Source/Editor/MeshPaint/Private/MeshPaintEdMode.cpp
--- a/Engine/Source/Editor/MeshPaint/Private/MeshPaintEdMode.cpp
+++ b/Engine/Source/Editor/MeshPaint/Private/MeshPaintEdMode.cpp
@@ -180,7 +180,7 @@ void FEdModeMeshPaint::Enter()
 	//When painting vertext colors we want to force the lod level of objects being painted to LOD0.
 	if( FMeshPaintSettings::Get().ResourceType == EMeshPaintResource::VertexColors )
 	{
-		ForceBestLOD();
+		ApplyOrRemoveForceBestLOD(/*bApply=*/ true);
 	}
 }
 
@@ -194,7 +194,7 @@ void FEdModeMeshPaint::Exit()
 	if( FMeshPaintSettings::Get().ResourceType == EMeshPaintResource::VertexColors )
 	{
 		ApplyVertexColorsToAllLODs();
-		ClearForcedLOD();
+		ApplyOrRemoveForceBestLOD(/*bApply=*/ false);
 	}
 
 	// The user can manipulate the editor selection lock flag in paint mode so we make sure to restore it here
@@ -224,7 +224,7 @@ void FEdModeMeshPaint::Exit()
 	}
 
 	// If the user has pending changes and the editor is not exiting, we want to do the commit for all the modified textures.
-	if( GetNumberOfPendingPaintChanges() > 0 && !GIsRequestingExit )
+	if ((GetNumberOfPendingPaintChanges() > 0) && !GIsRequestingExit)
 	{
 		CommitAllPaintedTextures();
 	}
@@ -954,7 +954,7 @@ void FEdModeMeshPaint::DoPaint( const FVector& InCameraOrigin,
 			}
 
 			bool bHasStaticMesh = false;
-			TArray<UMeshComponent*> MeshComponents;
+			TInlineComponentArray<UMeshComponent*> MeshComponents;
 			CurActor->GetComponents<UMeshComponent>(MeshComponents);
 			for (const auto& CurMeshComponent : MeshComponents)
 			{
@@ -1631,7 +1631,7 @@ void FEdModeMeshPaint::PaintMeshVertices(
 
 
 /** Paints mesh texture */
-void FEdModeMeshPaint::PaintMeshTexture( UMeshComponent* MeshComponent, const FMeshPaintParameters& Params, const bool bShouldApplyPaint, const FVector& ComponentSpaceCameraPosition, const FMatrix& ComponentToWorldMatrix, const float ComponentSpaceSquaredBrushRadius, const FVector& ComponentSpaceBrushPosition, const class IMeshPaintGeometryAdapter& GeometryInfo )
+void FEdModeMeshPaint::PaintMeshTexture( UMeshComponent* MeshComponent, const FMeshPaintParameters& Params, const bool bShouldApplyPaint, const FVector& ComponentSpaceCameraPosition, const FMatrix& ComponentToWorldMatrix, const float ComponentSpaceSquaredBrushRadius, const FVector& ComponentSpaceBrushPosition, const IMeshPaintGeometryAdapter& GeometryInfo )
 {
 	UTexture2D* TargetTexture2D = GetSelectedTexture();
 	
@@ -1792,21 +1792,21 @@ void FEdModeMeshPaint::PaintMeshTexture( UMeshComponent* MeshComponent, const FM
 	}
 
 
-	if ((TexturePaintingCurrentMeshComponent != NULL) && (TexturePaintingCurrentMeshComponent != MeshComponent))
+	if ((TexturePaintingCurrentMeshComponent != nullptr) && (TexturePaintingCurrentMeshComponent != MeshComponent))
 	{
 		// Mesh has changed, so finish up with our previous texture
 		FinishPaintingTexture();
 		bIsPainting = false;
 	}
 
-	if (TexturePaintingCurrentMeshComponent == NULL)
+	if (TexturePaintingCurrentMeshComponent == nullptr)
 	{
-		StartPaintingTexture( MeshComponent );
+		StartPaintingTexture(MeshComponent, GeometryInfo);
 	}
 
-	if (TexturePaintingCurrentMeshComponent != NULL)
+	if (TexturePaintingCurrentMeshComponent != nullptr)
 	{
-		PaintTexture( Params, InfluencedTriangles, ComponentToWorldMatrix, GeometryInfo );
+		PaintTexture(Params, InfluencedTriangles, ComponentToWorldMatrix, GeometryInfo);
 	}
 }
 
@@ -1814,7 +1814,7 @@ void FEdModeMeshPaint::PaintMeshTexture( UMeshComponent* MeshComponent, const FM
 
 
 /** Starts painting a texture */
-void FEdModeMeshPaint::StartPaintingTexture( UMeshComponent* InMeshComponent )
+void FEdModeMeshPaint::StartPaintingTexture(UMeshComponent* InMeshComponent, const IMeshPaintGeometryAdapter& GeometryInfo)
 {
 	check( InMeshComponent != NULL );
 	check( TexturePaintingCurrentMeshComponent == NULL );
@@ -1838,12 +1838,12 @@ void FEdModeMeshPaint::StartPaintingTexture( UMeshComponent* InMeshComponent )
 	{
 		bool bIsTextureUsed = DoesMeshComponentUseTexture(InMeshComponent, Texture2D);
 
-		if( !bIsTextureUsed && TextureData != NULL && TextureData->PaintRenderTargetTexture != NULL )
+		if (!bIsTextureUsed && (TextureData != nullptr) && (TextureData->PaintRenderTargetTexture != nullptr))
 		{
 			bIsTextureUsed = DoesMeshComponentUseTexture(InMeshComponent, TextureData->PaintRenderTargetTexture);
 		}
 
-		if( bIsTextureUsed == true && bStartedPainting == false )
+		if (bIsTextureUsed && !bStartedPainting)
 		{
 			bool bIsSourceTextureStreamedIn = Texture2D->IsFullyStreamedIn();
 
@@ -1951,10 +1951,11 @@ void FEdModeMeshPaint::StartPaintingTexture( UMeshComponent* InMeshComponent )
 
 		// @todo MeshPaint: Here we override the textures on the mesh with the render target.  The problem is that other meshes in the scene that use
 		//    this texture do not get the override. Do we want to extend this to all other selected meshes or maybe even to all meshes in the scene?
-		if( bIsTextureUsed == true && bStartedPainting == true && TextureData->PaintingMaterials.Contains( MaterialToCheck ) == false)
+		if (bIsTextureUsed && bStartedPainting && !TextureData->PaintingMaterials.Contains(MaterialToCheck))
 		{
 			TextureData->PaintingMaterials.AddUnique( MaterialToCheck ); 
-			MaterialToCheck->OverrideTexture(Texture2D, TextureData->PaintRenderTargetTexture, FeatureLevel);
+
+			GeometryInfo.ApplyOrRemoveTextureOverride(Texture2D, TextureData->PaintRenderTargetTexture);
 		}
 
 		MaterialIndex++;
@@ -2403,34 +2404,19 @@ void FEdModeMeshPaint::RestoreRenderTargets()
 }
 
 /** Clears all texture overrides for this component */
-void FEdModeMeshPaint::ClearMeshTextureOverrides(UMeshComponent* InMeshComponent)
+void FEdModeMeshPaint::ClearMeshTextureOverrides(const IMeshPaintGeometryAdapter& GeometryInfo, UMeshComponent* InMeshComponent)
 {
-	if (!InMeshComponent)
-	{
-		return;
-	}
-
-	const auto FeatureLevel = GetWorld()->FeatureLevel;
-
-	TArray<UMaterialInterface*> UsedMaterials;
-
-	// Get all the used materials for this StaticMeshComponent
-	InMeshComponent->GetUsedMaterials(UsedMaterials);
-
-	for( int32 MatIndex = 0; MatIndex < UsedMaterials.Num(); MatIndex++)
+	if (InMeshComponent != nullptr)
 	{
-		UMaterialInterface* Material = UsedMaterials[ MatIndex ];
+		TArray<UTexture*> UsedTextures;
+		InMeshComponent->GetUsedTextures(/*out*/ UsedTextures, EMaterialQualityLevel::High);
 
-		if( Material != NULL )
+		for (UTexture* Texture : UsedTextures)
 		{
-			TArray<UTexture*> UsedTextures;
-			Material->GetUsedTextures(UsedTextures, EMaterialQualityLevel::Num, false, FeatureLevel, false);
-
-			for( int32 UsedIndex = 0; UsedIndex < UsedTextures.Num(); UsedIndex++ )
+			if (UTexture2D* Texture2D = Cast<UTexture2D>(Texture))
 			{
-				//Reset the texture to it's default.
-				Material->OverrideTexture(UsedTextures[UsedIndex], NULL, FeatureLevel);
-			}		
+				GeometryInfo.ApplyOrRemoveTextureOverride(Texture2D, nullptr);
+			}
 		}
 	}
 }
@@ -2438,6 +2424,8 @@ void FEdModeMeshPaint::ClearMeshTextureOverrides(UMeshComponent* InMeshComponent
 /** Clears all texture overrides, removing any pending texture paint changes */
 void FEdModeMeshPaint::ClearAllTextureOverrides()
 {
+
+
 	for ( TMap< UTexture2D*, PaintTexture2DData >::TIterator It(PaintTargetData); It; ++It)
 	{
 		PaintTexture2DData* TextureData = &It.Value();
@@ -2445,7 +2433,7 @@ void FEdModeMeshPaint::ClearAllTextureOverrides()
 		for( int32 MaterialIndex = 0; MaterialIndex < TextureData->PaintingMaterials.Num(); MaterialIndex++)
 		{
 			UMaterialInterface* PaintingMaterialInterface = TextureData->PaintingMaterials[MaterialIndex];
-			PaintingMaterialInterface->OverrideTexture( TextureData->PaintingTexture2D, NULL, GetWorld()->FeatureLevel );
+			PaintingMaterialInterface->OverrideTexture( TextureData->PaintingTexture2D, NULL, GetWorld()->FeatureLevel );//findme
 		}
 
 		TextureData->PaintingMaterials.Empty();
@@ -2453,79 +2441,37 @@ void FEdModeMeshPaint::ClearAllTextureOverrides()
 }
 
 /** Sets all texture overrides available for the mesh. */
-void FEdModeMeshPaint::SetAllTextureOverrides(UMeshComponent* InMeshComponent)
+void FEdModeMeshPaint::SetAllTextureOverrides(const IMeshPaintGeometryAdapter& GeometryInfo, UMeshComponent* InMeshComponent)
 {
-	if (!InMeshComponent)
-	{
-		return;
-	}
-
-	const auto FeatureLevel = GetWorld()->FeatureLevel;
-
-	TArray<UMaterialInterface*> UsedMaterials;
-
-	// Get all the used materials for this StaticMeshComponent
-	InMeshComponent->GetUsedMaterials(UsedMaterials);
-
-	// Add the materials this actor uses to the list we maintain for ALL the selected actors, but only if
-	//  it does not appear in the list already.
-	for( int32 MatIndex = 0; MatIndex < UsedMaterials.Num(); MatIndex++)
+	if (InMeshComponent != nullptr)
 	{
-		UMaterialInterface* Material = UsedMaterials[ MatIndex ];
+		TArray<UTexture*> UsedTextures;
+		InMeshComponent->GetUsedTextures(/*out*/ UsedTextures, EMaterialQualityLevel::High);
 
-		if( Material != NULL )
+		for (UTexture* Texture : UsedTextures)
 		{
-			TArray<UTexture*> UsedTextures;
-			Material->GetUsedTextures(UsedTextures, EMaterialQualityLevel::Num, false, FeatureLevel, false);
-
-			for( int32 UsedIndex = 0; UsedIndex < UsedTextures.Num(); UsedIndex++ )
+			if (UTexture2D* Texture2D = Cast<UTexture2D>(Texture))
 			{
-
-				PaintTexture2DData* TextureData = GetPaintTargetData( (UTexture2D*)UsedTextures[ UsedIndex ] );
-
-				if(TextureData)
+				if (PaintTexture2DData* TextureData = GetPaintTargetData(Texture2D))
 				{
-					Material->OverrideTexture(UsedTextures[UsedIndex], TextureData->PaintRenderTargetTexture, FeatureLevel);
+					GeometryInfo.ApplyOrRemoveTextureOverride(Texture2D, TextureData->PaintRenderTargetTexture);
 				}
-			}		
+			}
 		}
 	}
-		
 }
 
 /** Sets the override for a specific texture for any materials using it in the mesh, clears the override if it has no overrides. */
-void FEdModeMeshPaint::SetSpecificTextureOverrideForMesh(UMeshComponent* InMeshComponent, UTexture* Texture)
+void FEdModeMeshPaint::SetSpecificTextureOverrideForMesh(const IMeshPaintGeometryAdapter& GeometryInfo, UTexture2D* Texture)
 {
-	PaintTexture2DData* TextureData = GetPaintTargetData( (UTexture2D*)Texture );
-
-	const auto FeatureLevel = GetWorld()->FeatureLevel;
-
-	// Check all the materials on the mesh to see if the user texture is there
-	int32 MaterialIndex = 0;
-	UMaterialInterface* MaterialToCheck = InMeshComponent->GetMaterial(MaterialIndex);
-	while( MaterialToCheck != NULL )
-	{
-		bool bIsTextureUsed = DoesMaterialUseTexture(MaterialToCheck,Texture);
-
-		if(bIsTextureUsed)
-		{
-			if(TextureData && TextureData->PaintingMaterials.Num() > 0)
-			{
-				// If there is texture data, that means we have an override ready, so set it. 
-				MaterialToCheck->OverrideTexture(Texture, TextureData->PaintRenderTargetTexture, FeatureLevel);
-			}
-			else
-			{
-				// If there is no data, then remove the override so we can at least see the texture without the changes to the other texture.
-					// This is important because overrides are shared between material instances with the same parent. We want to disable a override in place,
-					// making the action more comprehensive to the user.
-				MaterialToCheck->OverrideTexture(Texture, NULL, FeatureLevel);
-			}
-		}
-
-		++MaterialIndex;
-		MaterialToCheck = InMeshComponent->GetMaterial(MaterialIndex);
-	}
+	// If there is texture data, that means we have an override ready, so set it. 
+	// If there is no data, then remove the override so we can at least see the texture without the changes to the other texture.
+	// This is important because overrides are shared between material instances with the same parent. We want to disable a override in place,
+	// making the action more comprehensive to the user.
+	PaintTexture2DData* TextureData = GetPaintTargetData(Texture);
+	UTextureRenderTarget2D* TextureForOverrideOrNull = ((TextureData != nullptr) && (TextureData->PaintingMaterials.Num() > 0)) ? TextureData->PaintRenderTargetTexture : nullptr;
+
+	GeometryInfo.ApplyOrRemoveTextureOverride(Texture, TextureForOverrideOrNull);
 }
 
 int32 FEdModeMeshPaint::GetNumberOfPendingPaintChanges()
@@ -2712,7 +2658,7 @@ TArray<UMeshComponent*> FEdModeMeshPaint::GetSelectedMeshComponents() const
 			continue;
 		}
 
-		TArray<UMeshComponent*> MeshComponents;
+		TInlineComponentArray<UMeshComponent*> MeshComponents;
 		CurActor->GetComponents<UMeshComponent>(MeshComponents);
 
 		Result.Append(MeshComponents);
@@ -2726,7 +2672,7 @@ void FEdModeMeshPaint::SaveSettingsForActor( AActor* InActor )
 {
 	if( InActor != NULL )
 	{
-		TArray<UMeshComponent*> MeshComponents;
+		TInlineComponentArray<UMeshComponent*> MeshComponents;
 		InActor->GetComponents<UMeshComponent>(MeshComponents);
 
 		for( UMeshComponent* MeshComponent : MeshComponents )
@@ -2826,42 +2772,46 @@ void FEdModeMeshPaint::UpdateSettingsForMeshComponent( UMeshComponent* InMeshCom
 /** FEdMode: Handling SelectActor */
 bool FEdModeMeshPaint::Select( AActor* InActor, bool bInSelected )
 {
-	TArray<UStaticMeshComponent*> StaticMeshComponents;
-	InActor->GetComponents<UStaticMeshComponent>(StaticMeshComponents);
-	for(const auto& StaticMeshComponent : StaticMeshComponents)
+	TInlineComponentArray<UMeshComponent*> MeshComponents;
+	InActor->GetComponents<UMeshComponent>(MeshComponents);
+	for(const auto& MeshComponent : MeshComponents)
 	{
-		if( StaticMeshComponent != NULL )
+		if (MeshComponent != nullptr)
 		{
-			if( !bInSelected )
+			TSharedPtr<IMeshPaintGeometryAdapter> GeomInfo = FMeshPaintAdapterFactory::CreateAdapterForMesh(MeshComponent, PaintingMeshLODIndex, /*TODO: Shouldn't be part of the construction contract: FMeshPaintSettings::Get().UVChannel*/ 0);
+			if (GeomInfo.IsValid())
 			{
-				if(FMeshPaintSettings::Get().ResourceType == EMeshPaintResource::Texture)
-				{
-					// When un-selecting a mesh, save it's settings based on the current properties
-					ClearMeshTextureOverrides(StaticMeshComponent);
-					SaveSettingsForActor(InActor);
-				}
-				else if(FMeshPaintSettings::Get().ResourceType == EMeshPaintResource::VertexColors)
+				if (!bInSelected)
 				{
-					//Propagate painting to lower LODs and stop forcing the rendered mesh to LOD0.
-					ApplyVertexColorsToAllLODs(StaticMeshComponent);
-					ClearForcedLOD(StaticMeshComponent);
+					if (FMeshPaintSettings::Get().ResourceType == EMeshPaintResource::Texture)
 					{
-						FComponentReregisterContext ReregisterContext(StaticMeshComponent);
+						// When un-selecting a mesh, save it's settings based on the current properties
+						ClearMeshTextureOverrides(*GeomInfo, MeshComponent);
+						SaveSettingsForActor(InActor);
+					}
+					else if (FMeshPaintSettings::Get().ResourceType == EMeshPaintResource::VertexColors)
+					{
+						// Propagate painting to lower LODs and stop forcing the rendered mesh to LOD0.
+						ApplyVertexColorsToAllLODs(*GeomInfo, MeshComponent);
+						ApplyOrRemoveForceBestLOD(*GeomInfo, MeshComponent, /*bApply=*/ false);
+						{
+							FComponentReregisterContext ReregisterContext(MeshComponent);
+						}
 					}
 				}
-			}
-			else
-			{				
-				if(FMeshPaintSettings::Get().ResourceType == EMeshPaintResource::Texture)
-				{
-					SetAllTextureOverrides(StaticMeshComponent);
-				}
-				else if(FMeshPaintSettings::Get().ResourceType == EMeshPaintResource::VertexColors)
+				else
 				{
-					//Painting is done on LOD0 so force the mesh to render only LOD0.
-					ForceBestLOD(StaticMeshComponent);
+					if (FMeshPaintSettings::Get().ResourceType == EMeshPaintResource::Texture)
 					{
-						FComponentReregisterContext ReregisterContext(StaticMeshComponent);
+						SetAllTextureOverrides(*GeomInfo, MeshComponent);
+					}
+					else if (FMeshPaintSettings::Get().ResourceType == EMeshPaintResource::VertexColors)
+					{
+						//Painting is done on LOD0 so force the mesh to render only LOD0.
+						ApplyOrRemoveForceBestLOD(*GeomInfo, MeshComponent, /*bApply=*/ true);
+						{
+							FComponentReregisterContext ReregisterContext(MeshComponent);
+						}
 					}
 				}
 			}
@@ -3814,83 +3764,62 @@ void FEdModeMeshPaint::FixupInstanceVertexColors() const
 	}
 }
 
-void FEdModeMeshPaint::ForceBestLOD()
+void FEdModeMeshPaint::ApplyOrRemoveForceBestLOD(bool bApply)
 {
 	USelection& SelectedActors = *Owner->GetSelectedActors();
-	for( int32 CurSelectedActorIndex = 0; CurSelectedActorIndex < SelectedActors.Num(); ++CurSelectedActorIndex )
-	{
-		AActor* SelectedActor = Cast< AActor >( SelectedActors.GetSelectedObject( CurSelectedActorIndex ) );
-		if(SelectedActor != NULL)
-		{
-			TArray<UStaticMeshComponent*> StaticMeshComponents;
-			SelectedActor->GetComponents<UStaticMeshComponent>(StaticMeshComponents);
-			for(const auto& StaticMeshComponent : StaticMeshComponents)
-			{
-				ForceBestLOD(StaticMeshComponent);
-			}
-		}
-	}
-}
-
-void FEdModeMeshPaint::ForceBestLOD(UStaticMeshComponent* StaticMeshComponent)
-{
-	if( StaticMeshComponent )
-	{
-		//=0 means do not force the LOD.
-		//>0 means force the lod to x-1.
-		StaticMeshComponent->ForcedLodModel = 1;
-	}
-}
-
-void FEdModeMeshPaint::ClearForcedLOD()
-{
-	USelection& SelectedActors = *Owner->GetSelectedActors();
-	for( int32 CurSelectedActorIndex = 0; CurSelectedActorIndex < SelectedActors.Num(); ++CurSelectedActorIndex )
+	for (int32 CurSelectedActorIndex = 0; CurSelectedActorIndex < SelectedActors.Num(); ++CurSelectedActorIndex)
 	{
-		AActor* SelectedActor = Cast< AActor >( SelectedActors.GetSelectedObject( CurSelectedActorIndex ) );
-		if(SelectedActor != NULL)
+		if (AActor* SelectedActor = Cast<AActor>(SelectedActors.GetSelectedObject(CurSelectedActorIndex)))
 		{
-			TArray<UStaticMeshComponent*> StaticMeshComponents;
-			SelectedActor->GetComponents<UStaticMeshComponent>(StaticMeshComponents);
-			for(const auto& StaticMeshComponent : StaticMeshComponents)
+			TInlineComponentArray<UMeshComponent*> MeshComponents;
+			SelectedActor->GetComponents<UMeshComponent>(MeshComponents);
+			for (UMeshComponent* MeshComponent : MeshComponents)
 			{
-				ClearForcedLOD(StaticMeshComponent);
+				TSharedPtr<IMeshPaintGeometryAdapter> MeshAdapter = FMeshPaintAdapterFactory::CreateAdapterForMesh(MeshComponent, PaintingMeshLODIndex, /*TODO: Shouldn't be part of the construction contract: FMeshPaintSettings::Get().UVChannel*/ 0);
+				if (MeshAdapter.IsValid())
+				{
+					ApplyOrRemoveForceBestLOD(*MeshAdapter, MeshComponent, bApply);
+				}
 			}
 		}
 	}
 }
-
-void FEdModeMeshPaint::ClearForcedLOD(UStaticMeshComponent* StaticMeshComponent)
+void FEdModeMeshPaint::ApplyOrRemoveForceBestLOD(const IMeshPaintGeometryAdapter& GeometryInfo, UMeshComponent* MeshComponent, bool bApply)
 {
-	if( StaticMeshComponent )
+	if (UStaticMeshComponent* StaticMeshComponent = Cast<UStaticMeshComponent>(MeshComponent))
 	{
 		//=0 means do not force the LOD.
 		//>0 means force the lod to x-1.
-		StaticMeshComponent->ForcedLodModel = 0;
+		StaticMeshComponent->ForcedLodModel = bApply ? 1 : 0;
 	}
 }
 
 void FEdModeMeshPaint::ApplyVertexColorsToAllLODs()
 {	
-	// Find each static mesh component of any selected actors
+	// Find each mesh component of any selected actors
 	USelection& SelectedActors = *Owner->GetSelectedActors();
-	for( int32 CurSelectedActorIndex = 0; CurSelectedActorIndex < SelectedActors.Num(); ++CurSelectedActorIndex )
+	for (int32 CurSelectedActorIndex = 0; CurSelectedActorIndex < SelectedActors.Num(); ++CurSelectedActorIndex)
 	{
-		AActor* SelectedActor = Cast< AActor >( SelectedActors.GetSelectedObject( CurSelectedActorIndex ) );
-		if(SelectedActor != NULL)
+		if (AActor* SelectedActor = Cast<AActor>(SelectedActors.GetSelectedObject(CurSelectedActorIndex)))
 		{
-			TArray<UStaticMeshComponent*> StaticMeshComponents;
-			SelectedActor->GetComponents<UStaticMeshComponent>(StaticMeshComponents);
-			for(const auto& StaticMeshComponent : StaticMeshComponents)
+			TInlineComponentArray<UMeshComponent*> MeshComponents;
+			SelectedActor->GetComponents<UMeshComponent>(MeshComponents);
+			for (UMeshComponent* MeshComponent : MeshComponents)
 			{
-				ApplyVertexColorsToAllLODs(StaticMeshComponent);
+				TSharedPtr<IMeshPaintGeometryAdapter> MeshAdapter = FMeshPaintAdapterFactory::CreateAdapterForMesh(MeshComponent, PaintingMeshLODIndex, /*TODO: Shouldn't be part of the construction contract: FMeshPaintSettings::Get().UVChannel*/ 0);
+				if (MeshAdapter.IsValid())
+				{
+					ApplyVertexColorsToAllLODs(*MeshAdapter, MeshComponent);
+				}
 			}
 		}
 	}
 }
 
-void FEdModeMeshPaint::ApplyVertexColorsToAllLODs(UStaticMeshComponent* StaticMeshComponent)
+void FEdModeMeshPaint::ApplyVertexColorsToAllLODs(const IMeshPaintGeometryAdapter& GeometryInfo, UMeshComponent* InMeshComponent)
 {
+	UStaticMeshComponent* StaticMeshComponent = Cast<UStaticMeshComponent>(InMeshComponent);
+
 	// If a static mesh component was found, apply LOD0 painting to all lower LODs.
 	if( StaticMeshComponent && StaticMeshComponent->StaticMesh && FMeshPaintSettings::Get().ResourceType == EMeshPaintResource::VertexColors )
 	{
@@ -4006,19 +3935,15 @@ void FEdModeMeshPaint::CreateInstanceMaterialAndTexture() const
 void FEdModeMeshPaint::RemoveInstanceMaterialAndTexture() const
 {
 	USelection& SelectedActors = *Owner->GetSelectedActors();
-	for( int32 CurSelectedActorIndex = 0; CurSelectedActorIndex < SelectedActors.Num(); ++CurSelectedActorIndex )
+	for (int32 CurSelectedActorIndex = 0; CurSelectedActorIndex < SelectedActors.Num(); ++CurSelectedActorIndex)
 	{
-		AActor* SelectedActor = Cast< AActor >( SelectedActors.GetSelectedObject( CurSelectedActorIndex ) );
-		if(SelectedActor != NULL)
+		if (AActor* SelectedActor = Cast<AActor>(SelectedActors.GetSelectedObject(CurSelectedActorIndex)))
 		{
-			TArray<UMeshComponent*> StaticMeshComponents;
-			SelectedActor->GetComponents<UMeshComponent>(StaticMeshComponents);
-			for(const auto& MeshComponent : StaticMeshComponents)
+			TInlineComponentArray<UMeshComponent*> MeshComponents;
+			SelectedActor->GetComponents<UMeshComponent>(MeshComponents);
+			for (UMeshComponent* MeshComponent : MeshComponents)
 			{
-				if( MeshComponent != NULL )
-				{
-					// @todo: this function
-				}
+				// @todo: this function
 			}
 		}
 	}
@@ -4400,7 +4325,7 @@ void FEdModeMeshPaint::UpdateTexturePaintTargetList()
 				const int32 MaterialIndex = MeshData->SelectedMaterialIndex;
 
 				// we only operate on mesh components.
-				TArray<UMeshComponent*> MeshComponents;
+				TInlineComponentArray<UMeshComponent*> MeshComponents;
 				Actor->GetComponents<UMeshComponent>(MeshComponents);
 
 				int32 DefaultIndex = INDEX_NONE;
@@ -4785,14 +4710,18 @@ void FEdModeMeshPaint::Tick(FEditorViewportClient* ViewportClient,float DeltaTim
 {
 	FEdMode::Tick(ViewportClient,DeltaTime);
 
-	//Will set the texture override up for the selected texture, important for the drop down combo-list and selecting between material instances.
-	if(FMeshPaintSettings::Get().ResourceType == EMeshPaintResource::Texture)
+	// Will set the texture override up for the selected texture, important for the drop down combo-list and selecting between material instances.
+	if (FMeshPaintSettings::Get().ResourceType == EMeshPaintResource::Texture)
 	{
 		TArray<UMeshComponent*> MeshComponents = GetSelectedMeshComponents();
 
-		for( UMeshComponent* MeshComponent : MeshComponents )
+		for (UMeshComponent* MeshComponent : MeshComponents)
 		{
-			SetSpecificTextureOverrideForMesh(MeshComponent, GetSelectedTexture());
+			TSharedPtr<IMeshPaintGeometryAdapter> GeometryInfo = FMeshPaintAdapterFactory::CreateAdapterForMesh(MeshComponent, PaintingMeshLODIndex, /*FMeshPaintSettings::Get().UVChannel*/ 0);
+			if (GeometryInfo.IsValid())
+			{
+				SetSpecificTextureOverrideForMesh(*GeometryInfo, GetSelectedTexture());
+			}
 		}
 	}
 
@@ -4836,10 +4765,26 @@ void FEdModeMeshPaint::DuplicateTextureMaterialCombo()
 		if (MeshData != NULL)
 		{
 			int32 MaterialIndex = MeshData->SelectedMaterialIndex;
-			AStaticMeshActor* StaticMeshActor = Cast< AStaticMeshActor >( ActorBeingEdited.Get() );
-			if (StaticMeshActor != NULL && StaticMeshActor->GetStaticMeshComponent() != NULL)
+
+			AActor* Actor = ActorBeingEdited.Get();
+			TInlineComponentArray<UMeshComponent*> MeshComponents;
+			if (Actor != nullptr)
+			{
+				Actor->GetComponents<UMeshComponent>(MeshComponents);
+			}
+
+			if (MeshComponents.Num() > 0)
 			{
-				UMaterialInterface* MaterialToCheck = StaticMeshActor->GetStaticMeshComponent()->GetMaterial(MaterialIndex);
+				UMeshComponent* MeshComponent = MeshComponents[0];
+
+				TSharedPtr<IMeshPaintGeometryAdapter> MeshAdapter = FMeshPaintAdapterFactory::CreateAdapterForMesh(MeshComponent, PaintingMeshLODIndex, /*FMeshPaintSettings::Get().UVChannel*/ 0);
+				if (!MeshAdapter.IsValid())
+				{
+					return;
+				}
+
+
+				UMaterialInterface* MaterialToCheck = MeshComponent->GetMaterial(MaterialIndex);
 
 				bool bIsSourceTextureStreamedIn = SelectedTexture->IsFullyStreamedIn();
 
@@ -4931,13 +4876,12 @@ void FEdModeMeshPaint::DuplicateTextureMaterialCombo()
 				}
 
 				bool bMaterialChanged = false;
-				UStaticMeshComponent* SMComponent = StaticMeshActor->GetStaticMeshComponent();
-				ClearMeshTextureOverrides(SMComponent);
+				ClearMeshTextureOverrides(*MeshAdapter, MeshComponent);
 
-				SMComponent->SetMaterial(MaterialIndex,NewMaterialInstance);
-				UpdateSettingsForMeshComponent(SMComponent, SelectedTexture, NewTexture);
+				MeshComponent->SetMaterial(MaterialIndex, NewMaterialInstance);
+				UpdateSettingsForMeshComponent(MeshComponent, SelectedTexture, NewTexture);
 
-				SMComponent->MarkPackageDirty();
+				MeshComponent->MarkPackageDirty();
 			
 				ActorSelectionChangeNotify();
 			}
diff --git a/Engine/Source/Editor/MeshPaint/Private/MeshPaintModule.cpp b/Engine/Source/Editor/MeshPaint/Private/MeshPaintModule.cpp
--- a/Engine/Source/Editor/MeshPaint/Private/MeshPaintModule.cpp
+++ b/Engine/Source/Editor/MeshPaint/Private/MeshPaintModule.cpp
@@ -13,6 +13,26 @@
 //////////////////////////////////////////////////////////////////////////
 // IMeshPaintGeometryAdapter
 
+void IMeshPaintGeometryAdapter::DefaultApplyOrRemoveTextureOverride(UMeshComponent* InMeshComponent, UTexture* SourceTexture, UTexture* OverrideTexture)
+{
+	const ERHIFeatureLevel::Type FeatureLevel = InMeshComponent->GetWorld()->FeatureLevel;
+
+	// Check all the materials on the mesh to see if the user texture is there
+	int32 MaterialIndex = 0;
+	UMaterialInterface* MaterialToCheck = InMeshComponent->GetMaterial(MaterialIndex);
+	while (MaterialToCheck != nullptr)
+	{
+		const bool bIsTextureUsed = DoesMaterialUseTexture(MaterialToCheck, SourceTexture);
+		if (bIsTextureUsed)
+		{
+			MaterialToCheck->OverrideTexture(SourceTexture, OverrideTexture, FeatureLevel);
+		}
+
+		++MaterialIndex;
+		MaterialToCheck = InMeshComponent->GetMaterial(MaterialIndex);
+	}
+}
+
 void IMeshPaintGeometryAdapter::DefaultQueryPaintableTextures(int32 MaterialIndex, UMeshComponent* MeshComponent, int32& OutDefaultIndex, TArray<struct FPaintableTexture>& InOutTextureList)
 {
 	OutDefaultIndex = INDEX_NONE;
diff --git a/Engine/Source/Editor/MeshPaint/Private/MeshPaintStaticMeshAdapter.cpp b/Engine/Source/Editor/MeshPaint/Private/MeshPaintStaticMeshAdapter.cpp
--- a/Engine/Source/Editor/MeshPaint/Private/MeshPaintStaticMeshAdapter.cpp
+++ b/Engine/Source/Editor/MeshPaint/Private/MeshPaintStaticMeshAdapter.cpp
@@ -164,6 +164,11 @@ void FMeshPaintGeometryAdapterForStaticMeshes::QueryPaintableTextures(int32 Mate
 	DefaultQueryPaintableTextures(MaterialIndex, StaticMeshComponent, OutDefaultIndex, InOutTextureList);
 }
 
+void FMeshPaintGeometryAdapterForStaticMeshes::ApplyOrRemoveTextureOverride(UTexture* SourceTexture, UTexture* OverrideTexture) const
+{
+	DefaultApplyOrRemoveTextureOverride(StaticMeshComponent, SourceTexture, OverrideTexture);
+}
+
 //////////////////////////////////////////////////////////////////////////
 // FMeshPaintGeometryAdapterForStaticMeshesFactory
 
diff --git a/Engine/Source/Editor/MeshPaint/Private/MeshPaintStaticMeshAdapter.h b/Engine/Source/Editor/MeshPaint/Private/MeshPaintStaticMeshAdapter.h
--- a/Engine/Source/Editor/MeshPaint/Private/MeshPaintStaticMeshAdapter.h
+++ b/Engine/Source/Editor/MeshPaint/Private/MeshPaintStaticMeshAdapter.h
@@ -18,6 +18,7 @@ public:
 	virtual bool LineTraceComponent(struct FHitResult& OutHit, const FVector Start, const FVector End, const struct FCollisionQueryParams& Params) const override;
 	virtual void SphereIntersectTriangles(TArray<int32>& OutTriangles, const float ComponentSpaceSquaredBrushRadius, const FVector& ComponentSpaceBrushPosition) const override;
 	virtual void QueryPaintableTextures(int32 MaterialIndex, int32& OutDefaultIndex, TArray<struct FPaintableTexture>& InOutTextureList) override;
+	virtual void ApplyOrRemoveTextureOverride(UTexture* SourceTexture, UTexture* OverrideTexture) const override;
 
 protected:
 	UStaticMeshComponent* StaticMeshComponent;
diff --git a/Engine/Source/Editor/MeshPaint/Public/MeshPaintEdMode.h b/Engine/Source/Editor/MeshPaint/Public/MeshPaintEdMode.h
--- a/Engine/Source/Editor/MeshPaint/Public/MeshPaintEdMode.h
+++ b/Engine/Source/Editor/MeshPaint/Public/MeshPaintEdMode.h
@@ -6,6 +6,7 @@
 #include "GenericOctreePublic.h"
 #include "Engine/StaticMesh.h"
 
+class IMeshPaintGeometryAdapter;
 
 /** Mesh paint resource types */
 namespace EMeshPaintResource
@@ -434,20 +435,14 @@ public:
 	/** Applies vertex color painting on LOD0 to all lower LODs. */
 	void ApplyVertexColorsToAllLODs();
 
-	/** Forces all selected meshes to their best LOD. */
-	void ForceBestLOD();
+	/** Forces all selected meshes to their best LOD, or removes the forcing. */
+	void ApplyOrRemoveForceBestLOD(bool bApply);
 
-	/** Forces the passed static mesh component to it's best LOD. */
-	void ForceBestLOD(UStaticMeshComponent* StaticMeshComponent);
-
-	/** Stops forcing selected meshes to their best lod. */
-	void ClearForcedLOD();
-
-	/** Stops the passed static mesh component to it's best lod. */
-	void ClearForcedLOD(UStaticMeshComponent* StaticMeshComponent);
+	/** Forces the static mesh component to it's best LOD, or removes the forcing. */
+	void ApplyOrRemoveForceBestLOD(const IMeshPaintGeometryAdapter& GeometryInfo, UMeshComponent* InMeshComponent, bool bApply);
 
 	/** Applies vertex color painting on LOD0 to all lower LODs. */
-	void ApplyVertexColorsToAllLODs(UStaticMeshComponent* StaticMeshComponent);
+	void ApplyVertexColorsToAllLODs(const IMeshPaintGeometryAdapter& GeometryInfo, UMeshComponent* InMeshComponent);
 
 	/** Fills the vertex colors associated with the currently selected mesh*/
 	void FillInstanceVertexColors();
@@ -505,14 +500,14 @@ public:
 	void CommitAllPaintedTextures();
 
 	/** Clears all texture overrides for this component. */
-	void ClearMeshTextureOverrides(UMeshComponent* InMeshComponent);
+	void ClearMeshTextureOverrides(const IMeshPaintGeometryAdapter& GeometryInfo, UMeshComponent* InMeshComponent);
 
 	/** Clears all texture overrides, removing any pending texture paint changes */
 	void ClearAllTextureOverrides();
 
-	void SetAllTextureOverrides(UMeshComponent* InMeshComponent);
+	void SetAllTextureOverrides(const IMeshPaintGeometryAdapter& GeometryInfo, UMeshComponent* InMeshComponent);
 
-	void SetSpecificTextureOverrideForMesh(UMeshComponent* InMeshComponent, UTexture* Texture);
+	void SetSpecificTextureOverrideForMesh(const IMeshPaintGeometryAdapter& GeometryInfo, UTexture2D* Texture);
 
 	/** Used to tell the texture paint system that we will need to restore the rendertargets */
 	void RestoreRenderTargets();
@@ -652,7 +647,7 @@ private:
 	void PaintMeshVertices( UStaticMeshComponent* StaticMeshComponent, const FMeshPaintParameters& Params, const bool bShouldApplyPaint, FStaticMeshLODResources& LODModel, const FVector& ActorSpaceCameraPosition, const FMatrix& ActorToWorldMatrix, FPrimitiveDrawInterface* PDI, const float VisualBiasDistance );
 
 	/** Paints mesh texture */
-	void PaintMeshTexture( UMeshComponent* MeshComponent, const FMeshPaintParameters& Params, const bool bShouldApplyPaint, const FVector& ActorSpaceCameraPosition, const FMatrix& ActorToWorldMatrix, const float ActorSpaceSquaredBrushRadius, const FVector& ActorSpaceBrushPosition, const class IMeshPaintGeometryAdapter& GeometryInfo );
+	void PaintMeshTexture( UMeshComponent* MeshComponent, const FMeshPaintParameters& Params, const bool bShouldApplyPaint, const FVector& ActorSpaceCameraPosition, const FMatrix& ActorToWorldMatrix, const float ActorSpaceSquaredBrushRadius, const FVector& ActorSpaceBrushPosition, const IMeshPaintGeometryAdapter& GeometryInfo );
 
 	/** Forces real-time perspective viewports */
 	void ForceRealTimeViewports( const bool bEnable, const bool bStoreCurrentState );
@@ -661,13 +656,13 @@ private:
 	void SetViewportShowFlags( const bool bAllowColorViewModes, FEditorViewportClient& Viewport );
 
 	/** Starts painting a texture */
-	void StartPaintingTexture( UMeshComponent* InMeshComponent );
+	void StartPaintingTexture(UMeshComponent* InMeshComponent, const IMeshPaintGeometryAdapter& GeometryInfo);
 
 	/** Paints on a texture */
 	void PaintTexture( const FMeshPaintParameters& InParams,
 					   const TArray< int32 >& InInfluencedTriangles,
 					   const FMatrix& InActorToWorldMatrix,
-					   const class IMeshPaintGeometryAdapter& GeometryInfo);
+					   const IMeshPaintGeometryAdapter& GeometryInfo);
 
 	/** Finishes painting a texture */
 	void FinishPaintingTexture();
diff --git a/Engine/Source/Editor/MeshPaint/Public/MeshPaintModule.h b/Engine/Source/Editor/MeshPaint/Public/MeshPaintModule.h
--- a/Engine/Source/Editor/MeshPaint/Public/MeshPaintModule.h
+++ b/Engine/Source/Editor/MeshPaint/Public/MeshPaintModule.h
@@ -13,14 +13,15 @@ public:
 	virtual bool Construct(UMeshComponent* InComponent, int32 InPaintingMeshLODIndex, int32 InUVChannelIndex) = 0;
 	virtual int32 GetNumTexCoords() const = 0;
 	virtual void GetTriangleInfo(int32 TriIndex, struct FTexturePaintTriangleInfo& OutTriInfo) const = 0;
-
 	virtual bool SupportsTexturePaint() const = 0;
 	virtual bool SupportsVertexPaint() const = 0;
 	virtual bool LineTraceComponent(struct FHitResult& OutHit, const FVector Start, const FVector End, const struct FCollisionQueryParams& Params) const = 0;
 	virtual void SphereIntersectTriangles(TArray<int32>& OutTriangles, const float ComponentSpaceSquaredBrushRadius, const FVector& ComponentSpaceBrushPosition) const = 0;
 	virtual void QueryPaintableTextures(int32 MaterialIndex, int32& OutDefaultIndex, TArray<struct FPaintableTexture>& InOutTextureList) = 0;
+	virtual void ApplyOrRemoveTextureOverride(UTexture* SourceTexture, UTexture* OverrideTexture) const = 0;
 	virtual ~IMeshPaintGeometryAdapter() {}
 
+	MESHPAINT_API static void DefaultApplyOrRemoveTextureOverride(UMeshComponent* InMeshComponent, UTexture* SourceTexture, UTexture* OverrideTexture);
 	MESHPAINT_API static void DefaultQueryPaintableTextures(int32 MaterialIndex, UMeshComponent* MeshComponent, int32& OutDefaultIndex, TArray<struct FPaintableTexture>& InOutTextureList);
 };
 
