Possible solution for UE-2244 (https://jira.ol.epicgames.net/browse/UE-2244)

Original code was by [REDACTED_PERSON], contacted him and the reason he wrote it orginally (2012) was with the assumption that it would be possible to paint to seperate LOD levels.

- Removed RemoveInstanceVertexColorsWorker(MeshPaintEdMode) function (unnecessary and very similar to RemoveInstanceVertexColorsFromLOD (StaticMeshComponent)

- Altered RemoveComponentInstanceVertexColors (MeshPaintEdMode), does original checks but now calls RemoveInstanceVertexColors (StaticMeshComponent) which has the desired outcome

diff --git a/Engine/Source/Editor/MeshPaint/Private/MeshPaintEdMode.cpp b/Engine/Source/Editor/MeshPaint/Private/MeshPaintEdMode.cpp
--- a/Engine/Source/Editor/MeshPaint/Private/MeshPaintEdMode.cpp
+++ b/Engine/Source/Editor/MeshPaint/Private/MeshPaintEdMode.cpp
@@ -3443,30 +3443,22 @@ void FEdModeMeshPaint::RemoveComponentInstanceVertexColors(UStaticMeshComponent*
 {
 	if( StaticMeshComponent != NULL && StaticMeshComponent->StaticMesh != NULL && StaticMeshComponent->StaticMesh->GetNumLODs() > PaintingMeshLODIndex )
 	{
-		// Make sure we have component-level LOD information
-		if( StaticMeshComponent->LODData.Num() > PaintingMeshLODIndex )
-		{
-			FStaticMeshComponentLODInfo* InstanceMeshLODInfo = &StaticMeshComponent->LODData[ PaintingMeshLODIndex ];
-
-			if(InstanceMeshLODInfo->OverrideVertexColors)
-			{
-				// @todo MeshPaint: Should make this undoable
+		// Mark the mesh component as modified
+		StaticMeshComponent->Modify();
 
-				// If this is called from the Remove button being clicked the SMC wont be in a Reregister context,
-				// but when it gets called from a Paste or Copy to Source operation it's already inside a more specific
-				// SMCRecreateScene context so we shouldn't put it inside another one.
-				if(StaticMeshComponent->IsRenderStateCreated())
-				{
-					// Detach all instances of this static mesh from the scene.
-					FComponentReregisterContext ComponentReregisterContext( StaticMeshComponent );
+		// If this is called from the Remove button being clicked the SMC wont be in a Reregister context,
+		// but when it gets called from a Paste or Copy to Source operation it's already inside a more specific
+		// SMCRecreateScene context so we shouldn't put it inside another one.
+		if (StaticMeshComponent->IsRenderStateCreated())
+		{
+			// Detach all instances of this static mesh from the scene.
+			FComponentReregisterContext ComponentReregisterContext(StaticMeshComponent);
 
-					RemoveInstanceVertexColorsWorker(StaticMeshComponent, InstanceMeshLODInfo);
-				}
-				else
-				{
-					RemoveInstanceVertexColorsWorker(StaticMeshComponent, InstanceMeshLODInfo);
-				}
-			}
+			StaticMeshComponent->RemoveInstanceVertexColors();
+		}
+		else
+		{
+			StaticMeshComponent->RemoveInstanceVertexColors();
 		}
 	}
 }
@@ -3483,27 +3475,6 @@ void FEdModeMeshPaint::RemoveInstanceVertexColors() const
 	}
 }
 
-/**
- * Does the work of removing instance vertex colors from a single static mesh component.
- *
- * @param	StaticMeshComponent		The SMC to remove vertex colors from.
- * @param	InstanceMeshLODInfo		The instance's LODInfo which stores the painted information to be cleared.
- */
-void FEdModeMeshPaint::RemoveInstanceVertexColorsWorker(UStaticMeshComponent *StaticMeshComponent, FStaticMeshComponentLODInfo *InstanceMeshLODInfo) const
-{
-	// Mark the mesh component as modified
-	StaticMeshComponent->Modify();
-
-	InstanceMeshLODInfo->ReleaseOverrideVertexColorsAndBlock();
-
-	// With no colors, there's no longer a reason to store vertex color positions. Remove them and count
-	// the component as up-to-date with the source mesh.
-	InstanceMeshLODInfo->PaintedVertices.Empty();
-	StaticMeshComponent->StaticMeshDerivedDataKey = StaticMeshComponent->StaticMesh->RenderData->DerivedDataKey;
-}
-
-
-
 /** Copies vertex colors associated with the currently selected mesh */
 void FEdModeMeshPaint::CopyInstanceVertexColors()
 {
diff --git a/Engine/Source/Editor/MeshPaint/Public/MeshPaintEdMode.h b/Engine/Source/Editor/MeshPaint/Public/MeshPaintEdMode.h
--- a/Engine/Source/Editor/MeshPaint/Public/MeshPaintEdMode.h
+++ b/Engine/Source/Editor/MeshPaint/Public/MeshPaintEdMode.h
@@ -791,14 +791,6 @@ private:
 	/** Temporary buffers used when copying/pasting colors */
 	TArray<FPerComponentVertexColorData> CopiedColorsByComponent;
 
-	/**
-	 * Does the work of removing instance vertex colors from a single static mesh component.
-	 *
-	 * @param	StaticMeshComponent		The SMC to remove vertex colors from.
-	 * @param	InstanceMeshLODInfo		The instance's LODInfo which stores the painted information to be cleared.
-	 */
-	void RemoveInstanceVertexColorsWorker(UStaticMeshComponent* StaticMeshComponent, FStaticMeshComponentLODInfo *InstanceMeshLODInfo) const;
-
 	/** Texture paint: Will hold a list of texture items that we can paint on */
 	TArray<FTextureTargetListInfo> TexturePaintTargetList;
 
