Several aspects of Spline Meshes were not being implemented in place of the generic Static Mesh code:

- Allowed for spline mesh painting by overriding the color vertices like static meshes do.
- Fixed improperly placed vertices, used the modified vertices of the splined mesh rather than the vertices of the base static mesh.
- Fixed splined mesh initialization issue.

UE-18511

diff --git a/Engine/Source/Editor/MeshPaint/Private/MeshPaintEdMode.cpp b/Engine/Source/Editor/MeshPaint/Private/MeshPaintEdMode.cpp
--- a/Engine/Source/Editor/MeshPaint/Private/MeshPaintEdMode.cpp
+++ b/Engine/Source/Editor/MeshPaint/Private/MeshPaintEdMode.cpp
@@ -30,6 +30,7 @@
 #include "Engine/StaticMeshActor.h"
 #include "Materials/MaterialInstanceConstant.h"
 #include "MeshPaintAdapterFactory.h"
+#include "Components/SplineMeshComponent.h"
 
 #define LOCTEXT_NAMESPACE "MeshPaint_Mode"
 
@@ -1508,7 +1509,7 @@ void FEdModeMeshPaint::PaintMeshVertices(
 					for( int32 TriVertexNum = 0; TriVertexNum < 3; ++TriVertexNum )
 					{
 						VertexIndices[ TriVertexNum ] = Indices[ TriIndex * 3 + TriVertexNum ];
-						TriVertices[ TriVertexNum ] = LODModel.PositionVertexBuffer.VertexPosition( VertexIndices[ TriVertexNum ] );
+						TriVertices[ TriVertexNum ] = GetMeshVertex( StaticMeshComponent, LODModel, VertexIndices[ TriVertexNum ] );
 					}
 
 					// Check to see if the triangle is front facing
@@ -1527,7 +1528,7 @@ void FEdModeMeshPaint::PaintMeshVertices(
 				{
 					// Grab the mesh vertex and transform it to world space
 					const int32 VertexIndex = CurIndexIt.GetIndex();
-					const FVector& ModelSpaceVertexPosition = LODModel.PositionVertexBuffer.VertexPosition( VertexIndex );
+					FVector	ModelSpaceVertexPosition = GetMeshVertex(StaticMeshComponent, LODModel, VertexIndex);
 					FVector WorldSpaceVertexPosition = ComponentToWorldMatrix.TransformPosition( ModelSpaceVertexPosition );
 
 					FColor OriginalVertexColor = FColor( 255, 255, 255 );
@@ -1705,7 +1706,7 @@ void FEdModeMeshPaint::PaintMeshTexture( UMeshComponent* MeshComponent, const FM
 			FBox Bounds; 
 			for (int32 VertIndex = 0; VertIndex < Indices.Num(); ++VertIndex)
 			{
-				FVector CurVector = LODModel.PositionVertexBuffer.VertexPosition( Indices[VertIndex] );
+				FVector	CurVector = GetMeshVertex( StaticMeshComponent, LODModel, Indices[ VertIndex ] );
 				if(VertIndex > 0)
 				{
 					Bounds.Min.X = FMath::Min<float>( Bounds.Min.X, CurVector.X );
@@ -1731,7 +1732,7 @@ void FEdModeMeshPaint::PaintMeshTexture( UMeshComponent* MeshComponent, const FM
 				for( int32 TriVertexNum = 0; TriVertexNum < 3; ++TriVertexNum )
 				{
 					const int32 VertexIndex = Indices[ TriIndex * 3 + TriVertexNum ];
-					MeshTri.Vertices[ TriVertexNum ] = LODModel.PositionVertexBuffer.VertexPosition( VertexIndex );
+					MeshTri.Vertices[ TriVertexNum ] = GetMeshVertex( StaticMeshComponent, LODModel, VertexIndex );
 				}
 				MeshTri.Index = TriIndex;
 				FBox TriBox;
@@ -1810,7 +1811,21 @@ void FEdModeMeshPaint::PaintMeshTexture( UMeshComponent* MeshComponent, const FM
 	}
 }
 
+/** Gets the vertex from the mesh, taking spline modifications into account */
+FVector FEdModeMeshPaint::GetMeshVertex( UStaticMeshComponent* StaticMeshComponent, FStaticMeshLODResources& LODModel, int32 VertexIndex )
+{
+	FVector OutPosition = LODModel.PositionVertexBuffer.VertexPosition( VertexIndex );
 
+	// Spline meshes must be shifted along the spline
+	if ( USplineMeshComponent* SplineMeshComponent = Cast<USplineMeshComponent>( StaticMeshComponent ))
+	{
+		const FTransform SliceTransform = SplineMeshComponent->CalcSliceTransform( USplineMeshComponent::GetAxisValue(OutPosition, SplineMeshComponent->ForwardAxis ));
+		USplineMeshComponent::GetAxisValue( OutPosition, SplineMeshComponent->ForwardAxis ) = 0;
+		OutPosition = SliceTransform.TransformPosition( OutPosition );
+	}
+
+	return OutPosition;
+}
 
 
 /** Starts painting a texture */
diff --git a/Engine/Source/Editor/MeshPaint/Public/MeshPaintEdMode.h b/Engine/Source/Editor/MeshPaint/Public/MeshPaintEdMode.h
--- a/Engine/Source/Editor/MeshPaint/Public/MeshPaintEdMode.h
+++ b/Engine/Source/Editor/MeshPaint/Public/MeshPaintEdMode.h
@@ -655,6 +655,9 @@ private:
 	/** Sets show flags for perspective viewports */
 	void SetViewportShowFlags( const bool bAllowColorViewModes, FEditorViewportClient& Viewport );
 
+	/** Gets the vertex from the mesh, taking spline modifications into account */
+	FVector GetMeshVertex( UStaticMeshComponent* StaticMeshComponent, FStaticMeshLODResources&, int32 );
+
 	/** Starts painting a texture */
 	void StartPaintingTexture(UMeshComponent* InMeshComponent, const IMeshPaintGeometryAdapter& GeometryInfo);
 
diff --git a/Engine/Source/Runtime/Engine/Private/Components/SplineMeshComponent.cpp b/Engine/Source/Runtime/Engine/Private/Components/SplineMeshComponent.cpp
--- a/Engine/Source/Runtime/Engine/Private/Components/SplineMeshComponent.cpp
+++ b/Engine/Source/Runtime/Engine/Private/Components/SplineMeshComponent.cpp
@@ -88,16 +88,15 @@ FVertexFactoryShaderParameters* FSplineMeshVertexFactory::ConstructShaderParamet
 //////////////////////////////////////////////////////////////////////////
 // SplineMeshSceneProxy
 
-
-
-void FSplineMeshSceneProxy::InitResources(USplineMeshComponent* InComponent, int32 InLODIndex)
+void FSplineMeshSceneProxy::InitResources( USplineMeshComponent* InComponent, int32 InLODIndex, FColorVertexBuffer* InOverrideColorVertexBuffer )
 {
 	// Initialize the static mesh's vertex factory.
-	ENQUEUE_UNIQUE_RENDER_COMMAND_THREEPARAMETER(
+	ENQUEUE_UNIQUE_RENDER_COMMAND_FOURPARAMETER(
 		InitSplineMeshVertexFactory,
 		FSplineMeshVertexFactory*, VertexFactory, LODResources[InLODIndex].VertexFactory,
-		const FStaticMeshLODResources*, RenderData, &InComponent->StaticMesh->RenderData->LODResources[InLODIndex],
+		FStaticMeshLODResources*, RenderData, &InComponent->StaticMesh->RenderData->LODResources[InLODIndex],
 		UStaticMesh*, Parent, InComponent->StaticMesh,
+		FColorVertexBuffer*, OverrideColorVertexBuffer, InOverrideColorVertexBuffer,
 		{
 		FLocalVertexFactory::DataType Data;
 
@@ -119,13 +118,19 @@ void FSplineMeshSceneProxy::InitResources(USplineMeshComponent* InComponent, int
 			RenderData->VertexBuffer.GetStride(),
 			VET_PackedNormal
 			);
-
-		if (RenderData->ColorVertexBuffer.GetNumVertices() > 0)
+		
+		FColorVertexBuffer* LODColorVertexBuffer = &RenderData->ColorVertexBuffer;
+		// Override the buffer if it has been changed
+		if ( OverrideColorVertexBuffer != NULL )
+		{
+			LODColorVertexBuffer = OverrideColorVertexBuffer;
+		}
+		if ( LODColorVertexBuffer->GetNumVertices() > 0 )
 		{
 			Data.ColorComponent = FVertexStreamComponent(
-				&RenderData->ColorVertexBuffer,
+				LODColorVertexBuffer,
 				0,	// Struct offset to color
-				RenderData->ColorVertexBuffer.GetStride(),
+				LODColorVertexBuffer->GetStride(),
 				VET_Color
 				);
 		}
diff --git a/Engine/Source/Runtime/Engine/Private/SplineMeshSceneProxy.cpp b/Engine/Source/Runtime/Engine/Private/SplineMeshSceneProxy.cpp
--- a/Engine/Source/Runtime/Engine/Private/SplineMeshSceneProxy.cpp
+++ b/Engine/Source/Runtime/Engine/Private/SplineMeshSceneProxy.cpp
@@ -47,7 +47,13 @@ FSplineMeshSceneProxy::FSplineMeshSceneProxy(USplineMeshComponent* InComponent)
 
 		LODResources.Add(VertexFactory);
 
-		InitResources(InComponent, LODIndex);
+		if( InComponent->LODData.IsValidIndex( LODIndex ))
+		{
+			const FStaticMeshComponentLODInfo& ComponentLODInfo = InComponent->LODData[ LODIndex ];
+			InitResources( InComponent, LODIndex, ComponentLODInfo.OverrideVertexColors );
+		}
+		else
+			InitResources( InComponent, LODIndex, NULL );
 	}
 }
 
diff --git a/Engine/Source/Runtime/Engine/Private/SplineMeshSceneProxy.h b/Engine/Source/Runtime/Engine/Private/SplineMeshSceneProxy.h
--- a/Engine/Source/Runtime/Engine/Private/SplineMeshSceneProxy.h
+++ b/Engine/Source/Runtime/Engine/Private/SplineMeshSceneProxy.h
@@ -149,7 +149,7 @@ public:
 
 	virtual ~FSplineMeshSceneProxy() override;
 
-	void InitResources(USplineMeshComponent* InComponent, int32 InLODIndex);
+	void InitResources(USplineMeshComponent* InComponent, int32 InLODIndex, FColorVertexBuffer*);
 
 	void ReleaseResources();
 
