Corrected tooltips

diff --git a/Engine/Source/Editor/MeshPaint/Private/SMeshPaint.cpp b/Engine/Source/Editor/MeshPaint/Private/SMeshPaint.cpp
--- a/Engine/Source/Editor/MeshPaint/Private/SMeshPaint.cpp
+++ b/Engine/Source/Editor/MeshPaint/Private/SMeshPaint.cpp
@@ -1394,7 +1394,7 @@ void SMeshPaint::Construct(const FArguments& InArgs, TSharedRef<FMeshPaintToolKi
 									.VAlign(VAlign_Center)
 									[
 										SNew(SButton)
-										.ToolTipText( LOCTEXT("FindSourceMeshInContentBrowser", "Find source mesh in content browser") )
+										.ToolTipText( LOCTEXT("FindSourceMeshInContentBrowser", "Find texture in content browser") )
 										.ContentPadding(0.0f) 
 										.OnClicked(this, &SMeshPaint::FindTextureInContentBrowserButtonClicked)
 										.IsEnabled( this, &SMeshPaint::IsSelectedTextureValid)
@@ -1428,7 +1428,7 @@ void SMeshPaint::Construct(const FArguments& InArgs, TSharedRef<FMeshPaintToolKi
 									.VAlign(VAlign_Center)
 									[
 										SNew(SButton)
-										.ToolTipText( LOCTEXT("SaveDirtyPackges", "Saves dirty source mesh packages associated with current actor selection") )
+										.ToolTipText( LOCTEXT("SaveDirtyPackges", "Saves dirty source texture assets associated with current actor selection") )
 										.ContentPadding(0.0f) 
 										.OnClicked(this, &SMeshPaint::SaveTextureButtonClicked)
 										.IsEnabled(this, &SMeshPaint::IsSelectedTextureDirty)
