Fix shadow variable warning:

Engine\\Source\\Editor\\MeshPaint\\Private\\MeshPaintStaticMeshAdapter.cpp(83,18) :  warning: declaration shadows a field of 'FMeshPaintGeometryAdapterForStaticMeshes' [-Wshadow]

diff --git a/Engine/Source/Editor/MeshPaint/Private/MeshPaintStaticMeshAdapter.cpp b/Engine/Source/Editor/MeshPaint/Private/MeshPaintStaticMeshAdapter.cpp
--- a/Engine/Source/Editor/MeshPaint/Private/MeshPaintStaticMeshAdapter.cpp
+++ b/Engine/Source/Editor/MeshPaint/Private/MeshPaintStaticMeshAdapter.cpp
@@ -80,8 +80,8 @@ void FMeshPaintGeometryAdapterForStaticMeshes::BuildOctree()
 
 	MeshTriOctree = MakeUnique<FMeshTriOctree>(Bounds.GetCenter(), Bounds.GetExtent().GetMax());
 
-	FIndexArrayView Indices = LODModel->IndexBuffer.GetArrayView();
-	const uint32 NumIndexBufferIndices = Indices.Num();
+	FIndexArrayView LocalIndices = LODModel->IndexBuffer.GetArrayView();
+	const uint32 NumIndexBufferIndices = LocalIndices.Num();
 	check(NumIndexBufferIndices % 3 == 0);
 	const uint32 NumTriangles = NumIndexBufferIndices / 3;
 
@@ -89,9 +89,9 @@ void FMeshPaintGeometryAdapterForStaticMeshes::BuildOctree()
 	{
 		// Grab the vertex indices and points for this triangle
 		FMeshTriangle MeshTri;
-		MeshTri.Vertices[0] = MeshVertices[Indices[TriIndex * 3 + 0]];
-		MeshTri.Vertices[1] = MeshVertices[Indices[TriIndex * 3 + 1]];
-		MeshTri.Vertices[2] = MeshVertices[Indices[TriIndex * 3 + 2]];
+		MeshTri.Vertices[0] = MeshVertices[LocalIndices[TriIndex * 3 + 0]];
+		MeshTri.Vertices[1] = MeshVertices[LocalIndices[TriIndex * 3 + 1]];
+		MeshTri.Vertices[2] = MeshVertices[LocalIndices[TriIndex * 3 + 2]];
 		MeshTri.Normal = FVector::CrossProduct(MeshTri.Vertices[1] - MeshTri.Vertices[0], MeshTri.Vertices[2] - MeshTri.Vertices[0]).GetSafeNormal();
 		MeshTri.Index = TriIndex;
 
