Fix shadowed variable

diff --git a/Engine/Source/Editor/MeshPaint/Private/MeshPaintStaticMeshAdapter.cpp b/Engine/Source/Editor/MeshPaint/Private/MeshPaintStaticMeshAdapter.cpp
--- a/Engine/Source/Editor/MeshPaint/Private/MeshPaintStaticMeshAdapter.cpp
+++ b/Engine/Source/Editor/MeshPaint/Private/MeshPaintStaticMeshAdapter.cpp
@@ -202,12 +202,12 @@ int32 FMeshPaintGeometryAdapterForStaticMeshes::GetNumTexCoords() const
 
 void FMeshPaintGeometryAdapterForStaticMeshes::GetTriangleInfo(int32 TriIndex, FTexturePaintTriangleInfo& OutTriInfo) const
 {
-	FIndexArrayView Indices = LODModel->IndexBuffer.GetArrayView();
+	FIndexArrayView VertexIndices = LODModel->IndexBuffer.GetArrayView();
 
 	// Grab the vertex indices and points for this triangle
 	for (int32 TriVertexNum = 0; TriVertexNum < 3; ++TriVertexNum)
 	{
-		const int32 VertexIndex = Indices[TriIndex * 3 + TriVertexNum];
+		const int32 VertexIndex = VertexIndices[TriIndex * 3 + TriVertexNum];
 		OutTriInfo.TriVertices[TriVertexNum] = MeshVertices[VertexIndex];
 		OutTriInfo.TriUVs[TriVertexNum] = LODModel->VertexBuffer.GetVertexUV(VertexIndex, UVChannelIndex);
 	}
