VR Editor: Vertex and texture painting support (UETOOL-668)
- You can now activate Mesh Paint and use the vertex/texture painting while in VR!
- Some features are not yet accessible (e.g. color picker)
- There is a known issue where the HMD view can flicker while painting on a texture

diff --git a/Engine/Source/Editor/MeshPaint/MeshPaint.Build.cs b/Engine/Source/Editor/MeshPaint/MeshPaint.Build.cs
--- a/Engine/Source/Editor/MeshPaint/MeshPaint.Build.cs
+++ b/Engine/Source/Editor/MeshPaint/MeshPaint.Build.cs
@@ -6,8 +6,6 @@ public class MeshPaint : ModuleRules
 {
     public MeshPaint(TargetInfo Target)
     {
-		PrivateIncludePathModuleNames.Add("AssetTools");
-
         PrivateDependencyModuleNames.AddRange(
             new string[] {
                 "AppFramework",
@@ -24,10 +22,22 @@ public class MeshPaint : ModuleRules
                 "EditorStyle",
                 "UnrealEd",
                 "RawMesh",
-                "SourceControl"
+                "SourceControl",
             } 
         );
 
-		DynamicallyLoadedModuleNames.Add("AssetTools");
+		PrivateIncludePathModuleNames.AddRange(
+			new string[]
+			{
+				"AssetTools",
+				"VREditor"
+			} );
+
+		DynamicallyLoadedModuleNames.AddRange(
+			new string[]
+			{
+				"AssetTools",
+				"VREditor"
+			} );
     }
 }
diff --git a/Engine/Source/Editor/MeshPaint/Private/MeshPaintEdMode.cpp b/Engine/Source/Editor/MeshPaint/Private/MeshPaintEdMode.cpp
--- a/Engine/Source/Editor/MeshPaint/Private/MeshPaintEdMode.cpp
+++ b/Engine/Source/Editor/MeshPaint/Private/MeshPaintEdMode.cpp
@@ -88,6 +88,7 @@ public:
 FEdModeMeshPaint::FEdModeMeshPaint() 
 	: FEdMode(),
 	  bIsPainting( false ),
+	  PaintingWithHandIndexInVR( INDEX_NONE ),
 	  bIsFloodFill( false ),
 	  bPushInstanceColorsToMesh( false ),
 	  PaintingStartTime( 0.0 ),
@@ -195,6 +196,17 @@ void FEdModeMeshPaint::Enter()
 	{
 		ApplyOrRemoveForceBestLOD(/*bApply=*/ true);
 	}
+
+	// Register to find out about editor modes being activated
+	GetModeManager()->OnEditorModeChanged().AddRaw( this, &FEdModeMeshPaint::OnEditorModeChanged );
+
+	// Register to find out about VR input events
+	IVREditorMode* VREditorMode = static_cast<IVREditorMode*>( GetModeManager()->GetActiveMode( IVREditorModule::Get().GetVREditorModeID() ) );
+	if( VREditorMode != nullptr )
+	{
+		VREditorMode->OnVRAction().RemoveAll( this );
+		VREditorMode->OnVRAction().AddRaw( this, &FEdModeMeshPaint::OnVRAction );
+	}
 }
 
 
@@ -202,6 +214,18 @@ void FEdModeMeshPaint::Enter()
 /** FEdMode: Called when the mode is exited */
 void FEdModeMeshPaint::Exit()
 {
+	// Unregister from event handlers
+	if( IVREditorModule::IsAvailable() )
+	{
+		IVREditorMode* VREditorMode = static_cast<IVREditorMode*>( GetModeManager()->GetActiveMode( IVREditorModule::Get().GetVREditorModeID() ) );
+		if( VREditorMode != nullptr )
+		{
+			VREditorMode->OnVRAction().RemoveAll( this );
+		}
+	}
+
+	GetModeManager()->OnEditorModeChanged().RemoveAll( this );
+
 	//If we're painting vertex colors then propagate the painting done on LOD0 to all lower LODs. 
 	//Then stop forcing the LOD level of the mesh to LOD0.
 	if( FMeshPaintSettings::Get().ResourceType == EMeshPaintResource::VertexColors )
@@ -271,6 +295,18 @@ void FEdModeMeshPaint::Exit()
 }
 
 
+void FEdModeMeshPaint::OnEditorModeChanged( FEdMode* EditorMode, bool bEntered )
+{
+	// VR Editor mode may have gone away, so re-register for events in case it's a different object than the one we were originally bound to
+	check( EditorMode != nullptr );
+	if( EditorMode->GetID() == IVREditorModule::Get().GetVREditorModeID() )
+	{
+		IVREditorMode* VREditorMode = static_cast<IVREditorMode*>( EditorMode );
+		VREditorMode->OnVRAction().RemoveAll( this );
+		VREditorMode->OnVRAction().AddRaw( this, &FEdModeMeshPaint::OnVRAction );
+	}
+}
+
 
 /** FEdMode: Called when the mouse is moved over the viewport */
 bool FEdModeMeshPaint::MouseMove( FEditorViewportClient* ViewportClient, FViewport* Viewport, int32 x, int32 y )
@@ -299,7 +335,7 @@ bool FEdModeMeshPaint::MouseMove( FEditorViewportClient* ViewportClient, FViewpo
 bool FEdModeMeshPaint::CapturedMouseMove( FEditorViewportClient* InViewportClient, FViewport* InViewport, int32 InMouseX, int32 InMouseY )
 {
 	// We only care about perspective viewports
-	if( InViewportClient->IsPerspective() && InViewportClient->EngineShowFlags.ModeWidgets )
+	if( InViewportClient->IsPerspective() )
 	{
 		if( bIsPainting )
 		{
@@ -316,12 +352,13 @@ bool FEdModeMeshPaint::CapturedMouseMove( FEditorViewportClient* InViewportClien
 			// Paint!
 			const bool bVisualCueOnly = false;
 			const EMeshPaintAction::Type PaintAction = GetPaintAction(InViewport);
+
 			// Apply stylus pressure
 			const float StrengthScale = InViewport->IsPenActive() ? InViewport->GetTabletPressure() : 1.f;
 
-			bool bAnyPaintAbleActorsUnderCursor = false;
+			bool bAnyPaintableActorsUnderCursor = false;
 
-			DoPaint( View->ViewMatrices.ViewOrigin, MouseViewportRay.GetOrigin(), MouseViewportRay.GetDirection(), NULL, PaintAction, bVisualCueOnly, StrengthScale, bAnyPaintAbleActorsUnderCursor );
+			DoPaint( View->ViewMatrices.ViewOrigin, MouseViewportRay.GetOrigin(), MouseViewportRay.GetDirection(), NULL, PaintAction, bVisualCueOnly, StrengthScale, /* Out */ bAnyPaintableActorsUnderCursor );
 
 			return true;
 		}
@@ -361,6 +398,7 @@ void FEdModeMeshPaint::EndPainting()
 	if(bIsPainting)
 	{
 		bIsPainting = false;
+		PaintingWithHandIndexInVR = INDEX_NONE;
 		FinishPaintingTexture();
 
 		// Rebuild any static meshes that we painted on last stroke
@@ -459,15 +497,16 @@ bool FEdModeMeshPaint::InputKey( FEditorViewportClient* InViewportClient, FViewp
 		}
 	}
 
-	// When painting we only care about perspective viewports where we are we are allowed to show mode widgets
-	if( !bIsAltDown && InViewportClient->IsPerspective() && InViewportClient->EngineShowFlags.ModeWidgets)
+	// When painting we only care about perspective viewports
+	if( !bIsAltDown && InViewportClient->IsPerspective() )
 	{
 		// Does the user want to paint right now?
 		const bool bUserWantsPaint = bIsLeftButtonDown && !bIsRightButtonDown && !bIsAltDown;
-		bool bAnyPaintAbleActorsUnderCursor = false;
+		bool bAnyPaintableActorsUnderCursor = false;
 
 		// Stop current tracking if the user is no longer painting
-		if( bIsPainting && !bUserWantsPaint )
+		if( bIsPainting && !bUserWantsPaint &&
+			( InKey == EKeys::LeftMouseButton || InKey == EKeys::RightMouseButton && InKey == EKeys::LeftAlt || InKey == EKeys::RightAlt ) )
 		{
 			bHandled = true;
 			EndPainting();
@@ -506,7 +545,7 @@ bool FEdModeMeshPaint::InputKey( FEditorViewportClient* InViewportClient, FViewp
 						const bool bVisualCueOnly = false;
 						const EMeshPaintAction::Type PaintAction = GetPaintAction(InViewport);
 						const float StrengthScale = 1.0f;
-						DoPaint( View->ViewMatrices.ViewOrigin, MouseViewportRay.GetOrigin(), MouseViewportRay.GetDirection(), NULL, PaintAction, bVisualCueOnly, StrengthScale, bAnyPaintAbleActorsUnderCursor );
+						DoPaint( View->ViewMatrices.ViewOrigin, MouseViewportRay.GetOrigin(), MouseViewportRay.GetDirection(), NULL, PaintAction, bVisualCueOnly, StrengthScale, /* Out */ bAnyPaintableActorsUnderCursor );
 					}
 				}
 				else
@@ -518,7 +557,7 @@ bool FEdModeMeshPaint::InputKey( FEditorViewportClient* InViewportClient, FViewp
 			}
 		}
 
-		if( !bAnyPaintAbleActorsUnderCursor )
+		if( !bAnyPaintableActorsUnderCursor )
 		{
 			bHandled = false;
 		}
@@ -1007,7 +1046,7 @@ void FEdModeMeshPaint::DoPaint( const FVector& InCameraOrigin,
 								const EMeshPaintAction::Type InPaintAction,
 								const bool bVisualCueOnly,
 								const float InStrengthScale,
-								OUT bool& bAnyPaintAbleActorsUnderCursor)
+								OUT bool& bAnyPaintableActorsUnderCursor)
 {
 	const float BrushRadius = GetBrushRadiiDefault();
 
@@ -1086,10 +1125,10 @@ void FEdModeMeshPaint::DoPaint( const FVector& InCameraOrigin,
 		}
 	}
 
-	bAnyPaintAbleActorsUnderCursor = (PaintableComponents.Num() > 0);
+	bAnyPaintableActorsUnderCursor = (PaintableComponents.Num() > 0);
 
 	// Are we actually applying paint here?
-	const bool bShouldApplyPaint = bAnyPaintAbleActorsUnderCursor && ( (bIsPainting && !bVisualCueOnly) || 
+	const bool bShouldApplyPaint = bAnyPaintableActorsUnderCursor && ( (bIsPainting && !bVisualCueOnly) || 
 		(InPaintAction == EMeshPaintAction::Fill) || 
 		(InPaintAction == EMeshPaintAction::PushInstanceColorsToMesh) );
 
@@ -2555,12 +2594,8 @@ void FEdModeMeshPaint::Render( const FSceneView* View, FViewport* Viewport, FPri
 	/** Call parent implementation */
 	FEdMode::Render( View, Viewport, PDI );
 
-	// If this viewport does not support Mode widgets we will not draw it here.
 	FEditorViewportClient* ViewportClient = (FEditorViewportClient*)Viewport->GetClient();
-	if( ViewportClient && !ViewportClient->EngineShowFlags.ModeWidgets)
-	{
-		return;
-	}
+	check( ViewportClient != nullptr );
 
 	// We only care about perspective viewports
 	const bool bIsPerspectiveViewport = ( View->ViewMatrices.ProjMatrix.M[ 3 ][ 3 ] < ( 1.0f - SMALL_NUMBER ) );
@@ -2576,38 +2611,98 @@ void FEdModeMeshPaint::Render( const FSceneView* View, FViewport* Viewport, FPri
 		const bool bAllowColorViewModes = ( FMeshPaintSettings::Get().ResourceType != EMeshPaintResource::Texture );
 		SetViewportShowFlags( bAllowColorViewModes, *ViewportClient );
 
-		// Make sure the cursor is visible OR we're flood filling.  No point drawing a paint cue when there's no cursor.
-		if( Viewport->IsCursorVisible() || IsForceRendered())
-		{
 
-			if( !PDI->IsHitTesting() )
+		bool bHavePaintRay = false;
+		FVector PaintRayViewOrigin = FVector::ZeroVector;
+		FVector PaintRayStart = FVector::ZeroVector;
+		FVector PaintRayDirection = FVector::ZeroVector;
+
+		const int32 FirstHandIndex = PaintingWithHandIndexInVR != INDEX_NONE ? PaintingWithHandIndexInVR : 0;
+		const int32 LastHandIndex = PaintingWithHandIndexInVR != INDEX_NONE ? PaintingWithHandIndexInVR : 1;
+
+		for( int32 HandIndex = FirstHandIndex; HandIndex <= LastHandIndex; ++HandIndex )
+		{
+			// Check to see if VREditorMode is active
+			if( IVREditorModule::IsAvailable() )
 			{
-				// Grab the mouse cursor position
-				FIntPoint MousePosition;
-				Viewport->GetMousePos( MousePosition );
+				IVREditorMode* VREditorMode = static_cast<IVREditorMode*>( GetModeManager()->GetActiveMode( IVREditorModule::Get().GetVREditorModeID() ) );
+				if( VREditorMode != nullptr )
+				{
+					FVector LaserPointerStart, LaserPointerEnd;
+					if( VREditorMode->GetLaserPointer( HandIndex, /* Out */ LaserPointerStart, /* Out */ LaserPointerEnd ) )
+					{
+						const FVector LaserPointerDirection = ( LaserPointerEnd - LaserPointerStart ).GetSafeNormal();
 
-				// Is the mouse currently over the viewport? or flood filling
-				if(IsForceRendered() || ( MousePosition.X >= 0 && MousePosition.Y >= 0 && MousePosition.X < (int32)Viewport->GetSizeXY().X && MousePosition.Y < (int32)Viewport->GetSizeXY().Y) )
+						bHavePaintRay = true;
+						PaintRayViewOrigin = VREditorMode->GetHeadTransform().GetLocation();
+						PaintRayStart = LaserPointerStart;
+						PaintRayDirection = LaserPointerDirection;
+					}
+				}
+				else
 				{
-					// Compute a world space ray from the screen space mouse coordinates
-					FViewportCursorLocation MouseViewportRay( View, ViewportClient, MousePosition.X, MousePosition.Y );
+					// VREditor mode not currently active
+				}
+			}
 
 
-					// Unless "Flow" mode is enabled, we'll only draw a visual cue while rendering and won't
-					// do any actual painting.  When "Flow" is turned on we will paint here, too!
-					const bool bVisualCueOnly = !FMeshPaintSettings::Get().bEnableFlow;
-					float StrengthScale = FMeshPaintSettings::Get().bEnableFlow ? FMeshPaintSettings::Get().FlowAmount : 1.0f;
+			if( !bHavePaintRay )
+			{
+				// Make sure the cursor is visible OR we're flood filling.  No point drawing a paint cue when there's no cursor.
+				if( Viewport->IsCursorVisible() || IsForceRendered())
+				{
+					if( !PDI->IsHitTesting() )
+					{
+						// Grab the mouse cursor position
+						FIntPoint MousePosition;
+						Viewport->GetMousePos( MousePosition );
+
+						// Is the mouse currently over the viewport? or flood filling
+						if(IsForceRendered() || ( MousePosition.X >= 0 && MousePosition.Y >= 0 && MousePosition.X < (int32)Viewport->GetSizeXY().X && MousePosition.Y < (int32)Viewport->GetSizeXY().Y) )
+						{
+							// Compute a world space ray from the screen space mouse coordinates
+							FViewportCursorLocation MouseViewportRay( View, ViewportClient, MousePosition.X, MousePosition.Y );
 
-					// Apply stylus pressure if it's active
-					if( Viewport->IsPenActive() )
+							bHavePaintRay = true;
+							PaintRayViewOrigin = View->ViewMatrices.ViewOrigin;
+							PaintRayStart = MouseViewportRay.GetOrigin();
+							PaintRayDirection = MouseViewportRay.GetDirection();
+						}
+					}
+				}
+			}
+
+			if( bHavePaintRay )
+			{
+				// Unless "Flow" mode is enabled, we'll only draw a visual cue while rendering and won't
+				// do any actual painting.  When "Flow" is turned on we will paint here, too!
+				const bool bVisualCueOnly = !bIsPainting || ( PaintingWithHandIndexInVR == INDEX_NONE && !FMeshPaintSettings::Get().bEnableFlow );
+				float StrengthScale = ( PaintingWithHandIndexInVR == INDEX_NONE && FMeshPaintSettings::Get().bEnableFlow ) ? FMeshPaintSettings::Get().FlowAmount : 1.0f;
+
+				// Apply VR controller trigger pressure if we're painting in VR
+				if( bIsPainting && PaintingWithHandIndexInVR != INDEX_NONE )
+				{
+					IVREditorMode* VREditorMode = static_cast<IVREditorMode*>( GetModeManager()->GetActiveMode( IVREditorModule::Get().GetVREditorModeID() ) );
+					if( VREditorMode != nullptr )
 					{
-						StrengthScale *= Viewport->GetTabletPressure();
+						StrengthScale *= VREditorMode->GetHandInfo( PaintingWithHandIndexInVR ).SelectAndMoveTriggerValue;
+
+						// Make sure light press locking is enabled in VR for the hand we're painting with.  We don't want
+						// a full press to be converted to a move action, however locking will have been disabled already
+						// by the world interaction code after a light press happened on an object. 
+						VREditorMode->SetAllowTriggerLightPressLocking( PaintingWithHandIndexInVR, true );
 					}
+				}
 
-					const EMeshPaintAction::Type PaintAction = GetPaintAction(Viewport);
-					bool bAnyPaintAbleActorsUnderCursor = false;
-					DoPaint( View->ViewMatrices.ViewOrigin, MouseViewportRay.GetOrigin(), MouseViewportRay.GetDirection(), PDI, PaintAction, bVisualCueOnly, StrengthScale, bAnyPaintAbleActorsUnderCursor );
+				// Apply stylus pressure if it's active
+				else if( Viewport->IsPenActive() )
+				{
+					StrengthScale *= Viewport->GetTabletPressure();
 				}
+				
+				const EMeshPaintAction::Type PaintAction = GetPaintAction( Viewport );
+				bool bAnyPaintableActorsUnderCursor = false;
+				DoPaint( PaintRayViewOrigin, PaintRayStart, PaintRayDirection, PDI, PaintAction, bVisualCueOnly, StrengthScale, /* Out */ bAnyPaintableActorsUnderCursor );
 			}
 		}
 	}
@@ -3400,6 +3495,17 @@ EMeshPaintAction::Type FEdModeMeshPaint::GetPaintAction(FViewport* InViewport)
 	else
 	{
 		PaintAction = bShiftDown ? EMeshPaintAction::Erase : EMeshPaintAction::Paint;
+
+		// When painting using VR, allow the modifier button to activate Erase mode
+		if( PaintingWithHandIndexInVR != INDEX_NONE )
+		{
+			IVREditorMode* VREditorMode = static_cast<IVREditorMode*>( GetModeManager()->GetActiveMode( IVREditorModule::Get().GetVREditorModeID() ) );
+			if( VREditorMode != nullptr )
+			{
+				const bool bIsModifierPressed = VREditorMode->GetHandInfo( PaintingWithHandIndexInVR ).bIsModifierPressed;
+				PaintAction = bIsModifierPressed ? EMeshPaintAction::Erase : EMeshPaintAction::Paint;
+			}
+		}
 	}
 	return PaintAction;
 
@@ -4949,4 +5055,61 @@ TWeakObjectPtr<AActor> FEdModeMeshPaint::GetEditingActor() const
 }
 
 
+void FEdModeMeshPaint::OnVRAction( FEditorViewportClient& ViewportClient, const FVRAction VRAction, const EInputEvent Event, bool& bIsInputCaptured, bool& bWasHandled )
+{
+	IVREditorMode* VREditorMode = static_cast<IVREditorMode*>( GetModeManager()->GetActiveMode( IVREditorModule::Get().GetVREditorModeID() ) );
+	if( VREditorMode != nullptr )
+	{
+		// Absorb all "full press" events unless the user is clicking on UI to avoid objects getting moved around.
+		// @todo vreditor: We shouldn't have to use this trick.  This input system needs more flexibility in its public interface!
+		if( VRAction.ActionType == EVRActionType::SelectAndMove && bIsPainting && PaintingWithHandIndexInVR == VRAction.HandIndex && !VREditorMode->GetHandInfo( VRAction.HandIndex ).bIsHoveringOverUI )
+		{
+			bWasHandled = true;
+			bIsInputCaptured = true;
+		}
+
+		if( VRAction.ActionType == EVRActionType::SelectAndMove_LightlyPressed )
+		{
+			if( !bIsPainting && Event == IE_Pressed && !VREditorMode->GetHandInfo( VRAction.HandIndex ).bIsHoveringOverUI )
+			{
+				// NOTE: We intentionally don't capture input for "LightlyPressed" because we want to allow objects to still
+				// become selected when clicked on.  We avoid allowing the objects to be moved by calling the 
+				// SetAllowTriggerLightPressLocking() function in the next tick (after VR Editor will have disabled locking.)
+				// @todo vreditor: We shouldn't have to use this trick.  This input system needs more flexibility in its public interface!
+				// 				bWasHandled = true;
+				// 				bIsInputCaptured = true;
+
+				StartPainting();
+				PaintingWithHandIndexInVR = VRAction.HandIndex;
+
+				// Go ahead and paint immediately
+				FVector LaserPointerStart, LaserPointerEnd;
+				if( VREditorMode->GetLaserPointer( PaintingWithHandIndexInVR, /* Out */ LaserPointerStart, /* Out */ LaserPointerEnd ) )
+				{
+					const FVector LaserPointerDirection = ( LaserPointerEnd - LaserPointerStart ).GetSafeNormal();
+
+					// Apply VR controller trigger pressure
+					float StrengthScale = VREditorMode->GetHandInfo( PaintingWithHandIndexInVR ).SelectAndMoveTriggerValue;
+
+					// Paint!
+					const bool bVisualCueOnly = false;
+					const EMeshPaintAction::Type PaintAction = GetPaintAction( ViewportClient.Viewport );
+					bool bAnyPaintableActorsUnderCursor = false;
+					DoPaint( VREditorMode->GetHeadTransform().GetLocation(), LaserPointerStart, LaserPointerDirection, NULL, PaintAction, bVisualCueOnly, StrengthScale, /* Out */ bAnyPaintableActorsUnderCursor );
+				}
+			}
+
+			// Stop current tracking if the user is no longer painting
+			else if( bIsPainting && Event == IE_Released && PaintingWithHandIndexInVR == VRAction.HandIndex )
+			{
+				EndPainting();
+
+				// @todo vreditor: We shouldn't have to use this trick.  This input system needs more flexibility in its public interface!
+				// 				bWasHandled = true;
+				// 				bIsInputCaptured = false;
+			}
+		}
+	}
+}
+
 #undef LOCTEXT_NAMESPACE
diff --git a/Engine/Source/Editor/MeshPaint/Public/MeshPaintEdMode.h b/Engine/Source/Editor/MeshPaint/Public/MeshPaintEdMode.h
--- a/Engine/Source/Editor/MeshPaint/Public/MeshPaintEdMode.h
+++ b/Engine/Source/Editor/MeshPaint/Public/MeshPaintEdMode.h
@@ -5,6 +5,8 @@
 #include "GenericOctree.h"
 #include "GenericOctreePublic.h"
 #include "Engine/StaticMesh.h"
+#include "IVREditorModule.h"
+
 
 class IMeshPaintGeometryAdapter;
 
@@ -399,6 +401,9 @@ public:
 	// End of FEdMode interface
 
 
+	/** Called when an editor mode is entered or exited */
+	void OnEditorModeChanged( FEdMode* EditorMode, bool bEntered );
+
 	/** Returns true if we need to force a render/update through based fill/copy */
 	bool IsForceRendered (void) const;
 
@@ -709,11 +714,17 @@ private:
 	/** Called when an asset is about to be reimported */
 	void OnPreReimportAsset(UObject* Object);
 
+	/** Called when the user presses a button on their motion controller device */
+	void OnVRAction( FEditorViewportClient& ViewportClient, const FVRAction VRAction, const EInputEvent Event, bool& bIsInputCaptured, bool& bOutWasHandled ); 
+
 private:
 
 	/** Whether we're currently painting */
 	bool bIsPainting;
 
+	/** When painting in VR, this is the hand index that we're painting with.  Otherwise INDEX_NONE. */
+	int32 PaintingWithHandIndexInVR;
+
 	/** Whether we are doing a flood fill the next render */
 	bool bIsFloodFill;
 
diff --git a/Engine/Source/Editor/VREditor/Public/IVREditorMode.h b/Engine/Source/Editor/VREditor/Public/IVREditorMode.h
--- a/Engine/Source/Editor/VREditor/Public/IVREditorMode.h
+++ b/Engine/Source/Editor/VREditor/Public/IVREditorMode.h
@@ -75,6 +75,13 @@ public:
 	 */
 	virtual bool GetLaserPointer( const int32 HandIndex, FVector& LaserPointerStart, FVector& LaserPointerEnd, const bool bEvenIfUIIsInFront = false, const float LaserLengthOverride = 0.0f ) const = 0;
 
+	/**
+	 * Gets the world space transform of the HMD (head)
+	 *
+	 * @return	World space space HMD transform
+	 */
+	virtual FTransform GetHeadTransform() const = 0;
+
 	/**
 	 * Gets our avatar's mesh actor
 	 *
@@ -103,5 +110,40 @@ public:
 		first, as there may be sibling subscribers that handle the event before you! */
 	DECLARE_EVENT_FiveParams( IVREditorMode, FOnVRHoverUpdate, class FEditorViewportClient& /* ViewportClient */, const int32 /* HandIndex */, FVector& /* HoverImpactPoint */, bool& /* bIsHoveringOverUI */, bool& /* bWasHandled */ );
 	virtual FOnVRHoverUpdate& OnVRHoverUpdate() = 0;
+
+
+	struct FHandInfo
+	{
+		/** Is the Modifier button held down? */
+		bool bIsModifierPressed;
+
+		/** How much the trigger is held down for the 'select and move' action */
+		float SelectAndMoveTriggerValue;
+
+		/** True if we're hovering over UI right now.  When hovering over UI, we don't bother drawing a see-thru laser pointer */
+		bool bIsHoveringOverUI;
+	};
+
+
+	/**
+	 * Gets some information about the specified hand (motion controller)
+	 *
+	 * @param	HandIndex	Which hand to query info for
+	 *
+	 * @return	A structure containing information about the hand's current state
+	 */
+	virtual FHandInfo GetHandInfo( const int32 HandIndex ) const = 0;
+
+	/**
+	 * Enables or disable "light press" locking for the specified hand.  This is an advanced feature.
+	 * It should be true if we allow locking of the lightly pressed state for the current event.  This can 
+	 * be useful to turn off in cases where you always want a full press to override the light press, even 
+	 * if it was a very slow press
+	 *
+	 * @param	HandIndex	The hand to change light press locking state for
+	 * @param	bAllow		Enables or disables light press locking for this hand
+	 */
+	virtual void SetAllowTriggerLightPressLocking( const int32 HandIndex, const bool bAllow ) = 0;
+
 };
 
diff --git a/Engine/Source/Editor/VREditor/VREditorMode.cpp b/Engine/Source/Editor/VREditor/VREditorMode.cpp
--- a/Engine/Source/Editor/VREditor/VREditorMode.cpp
+++ b/Engine/Source/Editor/VREditor/VREditorMode.cpp
@@ -1580,6 +1580,9 @@ bool FVREditorMode::HandleInputAxis(FEditorViewportClient* ViewportClient, FView
 		{
 			if( VRAction.ActionType == EVRActionType::TriggerAxis )
 			{
+				// Store latest trigger value amount
+				Hand.SelectAndMoveTriggerValue = Delta;
+
 				// Synthesize "lightly pressed" events for the trigger
 				{
 					const bool bIsFullPressAlreadyCapturing = Hand.IsInputCaptured[ (int32)EVRActionType::SelectAndMove ];
@@ -1835,6 +1838,24 @@ bool FVREditorMode::IsInputCaptured( const FVRAction VRAction ) const
 }
 
 
+IVREditorMode::FHandInfo FVREditorMode::GetHandInfo( const int32 HandIndex ) const
+{
+	const FVirtualHand& Hand = GetVirtualHand( HandIndex );
+
+	IVREditorMode::FHandInfo HandInfo;
+	HandInfo.bIsModifierPressed = Hand.bIsModifierPressed;
+	HandInfo.SelectAndMoveTriggerValue = Hand.SelectAndMoveTriggerValue;
+	HandInfo.bIsHoveringOverUI = Hand.bIsHoveringOverUI;
+	return HandInfo;
+}
+
+
+void FVREditorMode::SetAllowTriggerLightPressLocking( const int32 HandIndex, const bool bAllow )
+{
+	GetVirtualHand( HandIndex ).bAllowTriggerLightPressLocking = bAllow;
+}
+
+
 FHitResult FVREditorMode::GetHitResultFromLaserPointer( int32 HandIndex, TArray<AActor*>* OptionalListOfIgnoredActors, const bool bIgnoreGizmos, const bool bEvenIfUIIsInFront, const float LaserLengthOverride )
 {
 	FHitResult BestHitResult;
diff --git a/Engine/Source/Editor/VREditor/VREditorMode.h b/Engine/Source/Editor/VREditor/VREditorMode.h
--- a/Engine/Source/Editor/VREditor/VREditorMode.h
+++ b/Engine/Source/Editor/VREditor/VREditorMode.h
@@ -70,6 +70,7 @@ public:
 
 	// IVREditorMode interface
 	virtual bool GetLaserPointer( const int32 HandIndex, FVector& LaserPointerStart, FVector& LaserPointerEnd, const bool bEvenIfUIIsInFront = false, const float LaserLengthOverride = 0.0f) const override;
+	virtual FTransform GetHeadTransform() const override;
 	virtual bool IsInputCaptured( const FVRAction VRAction ) const;
 	virtual AActor* GetAvatarMeshActor() override
 	{
@@ -83,6 +84,8 @@ public:
 	{
 		return OnVRHoverUpdateEvent;
 	}
+	virtual FHandInfo GetHandInfo( const int32 HandIndex ) const override;
+	virtual void SetAllowTriggerLightPressLocking( const int32 HandIndex, const bool bAllow ) override;
 
 	/** Returns the Unreal controller ID for the motion controllers we're using */
 	int32 GetMotionControllerID() const
@@ -158,9 +161,6 @@ public:
 	/** Gets the transform of the user's HMD in room space */
 	FTransform GetRoomSpaceHeadTransform() const;
 
-	/** Gets the transform of the user's HMD in world space */
-	FTransform GetHeadTransform() const;
-
 	/** Spawns a transient actor that we can use in the editor world (templated for convenience) */
 	template<class T>
 	T* SpawnTransientSceneActor( const FString& ActorName, const bool bWithSceneComponent ) const
diff --git a/Engine/Source/Editor/VREditor/VirtualHand.h b/Engine/Source/Editor/VREditor/VirtualHand.h
--- a/Engine/Source/Editor/VREditor/VirtualHand.h
+++ b/Engine/Source/Editor/VREditor/VirtualHand.h
@@ -102,6 +102,9 @@ struct FVirtualHand
 	/** True if trigger has been fully released since the last press */
 	bool bHasTriggerBeenReleasedSinceLastPress;
 
+	/** Current trigger pressed amount for 'select and move' (0.0 - 1.0) */
+	float SelectAndMoveTriggerValue;
+
 
 	//
 	// Trackpad support
@@ -287,6 +290,7 @@ struct FVirtualHand
 		bAllowTriggerLightPressLocking = true;
 		RealTimeTriggerWasLightlyPressed = 0.0;
 		bHasTriggerBeenReleasedSinceLastPress = true;
+		SelectAndMoveTriggerValue = 0.0f;
 
 		bIsTouchingTrackpad = false;
 		bIsModifierPressed = false;
