VR Editor:
- First extensibility stage for interaction
- Fixed world space rotation with gizmo handle

diff --git a/Engine/Source/Editor/MeshPaint/MeshPaint.Build.cs b/Engine/Source/Editor/MeshPaint/MeshPaint.Build.cs
--- a/Engine/Source/Editor/MeshPaint/MeshPaint.Build.cs
+++ b/Engine/Source/Editor/MeshPaint/MeshPaint.Build.cs
@@ -23,21 +23,21 @@ public class MeshPaint : ModuleRules
                 "UnrealEd",
                 "RawMesh",
                 "SourceControl",
-            } 
+                "ViewportInteraction",
+                "VREditor"
+            }
         );
 
 		PrivateIncludePathModuleNames.AddRange(
 			new string[]
 			{
-				"AssetTools",
-				"VREditor"
-			} );
+				"AssetTools"
+            });
 
 		DynamicallyLoadedModuleNames.AddRange(
 			new string[]
 			{
-				"AssetTools",
-				"VREditor"
-			} );
+				"AssetTools"
+            });
     }
 }
diff --git a/Engine/Source/Editor/MeshPaint/Private/MeshPaintEdMode.cpp b/Engine/Source/Editor/MeshPaint/Private/MeshPaintEdMode.cpp
--- a/Engine/Source/Editor/MeshPaint/Private/MeshPaintEdMode.cpp
+++ b/Engine/Source/Editor/MeshPaint/Private/MeshPaintEdMode.cpp
@@ -20,7 +20,6 @@
 #include "Editor/LevelEditor/Public/SLevelViewport.h"
 #include "MessageLog.h"
 
-
 #include "Runtime/Engine/Classes/PhysicsEngine/BodySetup.h"
 #include "SMeshPaint.h"
 #include "ComponentReregisterContext.h"
@@ -33,6 +32,9 @@
 #include "MeshPaintAdapterFactory.h"
 #include "Components/SplineMeshComponent.h"
 
+#include "ViewportWorldInteraction.h"
+#include "VREditorInteractor.h"
+
 #define LOCTEXT_NAMESPACE "MeshPaint_Mode"
 
 DEFINE_LOG_CATEGORY_STATIC(LogMeshPaintEdMode, Log, All);
@@ -88,7 +90,7 @@ public:
 FEdModeMeshPaint::FEdModeMeshPaint() 
 	: FEdMode(),
 	  bIsPainting( false ),
-	  PaintingWithHandIndexInVR( INDEX_NONE ),
+	  PaintingWithInteractorInVR( nullptr ),
 	  bIsFloodFill( false ),
 	  bPushInstanceColorsToMesh( false ),
 	  PaintingStartTime( 0.0 ),
@@ -204,13 +206,11 @@ void FEdModeMeshPaint::Enter()
 	IVREditorMode* VREditorMode = static_cast<IVREditorMode*>( GetModeManager()->GetActiveMode( IVREditorModule::Get().GetVREditorModeID() ) );
 	if( VREditorMode != nullptr )
 	{
-		VREditorMode->OnVRAction().RemoveAll( this );
-		VREditorMode->OnVRAction().AddRaw( this, &FEdModeMeshPaint::OnVRAction );
+		VREditorMode->GetWorldInteraction().OnViewportInteractionInputAction().RemoveAll( this );
+		VREditorMode->GetWorldInteraction().OnViewportInteractionInputAction().AddRaw( this, &FEdModeMeshPaint::OnVRAction );
 	}
 }
 
-
-
 /** FEdMode: Called when the mode is exited */
 void FEdModeMeshPaint::Exit()
 {
@@ -220,7 +220,7 @@ void FEdModeMeshPaint::Exit()
 		IVREditorMode* VREditorMode = static_cast<IVREditorMode*>( GetModeManager()->GetActiveMode( IVREditorModule::Get().GetVREditorModeID() ) );
 		if( VREditorMode != nullptr )
 		{
-			VREditorMode->OnVRAction().RemoveAll( this );
+			VREditorMode->GetWorldInteraction().OnViewportInteractionInputAction().RemoveAll( this );
 		}
 	}
 
@@ -299,11 +299,16 @@ void FEdModeMeshPaint::OnEditorModeChanged( FEdMode* EditorMode, bool bEntered )
 {
 	// VR Editor mode may have gone away, so re-register for events in case it's a different object than the one we were originally bound to
 	check( EditorMode != nullptr );
-	if( EditorMode->GetID() == IVREditorModule::Get().GetVREditorModeID() )
+	if( bEntered && EditorMode->GetID() == IVREditorModule::Get().GetVREditorModeID() )
 	{
 		IVREditorMode* VREditorMode = static_cast<IVREditorMode*>( EditorMode );
-		VREditorMode->OnVRAction().RemoveAll( this );
-		VREditorMode->OnVRAction().AddRaw( this, &FEdModeMeshPaint::OnVRAction );
+		if ( VREditorMode != nullptr )
+		{
+			VREditorMode->GetWorldInteraction().OnViewportInteractionInputAction().RemoveAll( this );
+			VREditorMode->GetWorldInteraction().OnViewportInteractionInputAction().AddRaw( this, &FEdModeMeshPaint::OnVRAction );
+		}
+
+		PaintingWithInteractorInVR = nullptr;
 	}
 }
 
@@ -398,7 +403,7 @@ void FEdModeMeshPaint::EndPainting()
 	if(bIsPainting)
 	{
 		bIsPainting = false;
-		PaintingWithHandIndexInVR = INDEX_NONE;
+		PaintingWithInteractorInVR = nullptr;
 		FinishPaintingTexture();
 
 		// Rebuild any static meshes that we painted on last stroke
@@ -2617,19 +2622,22 @@ void FEdModeMeshPaint::Render( const FSceneView* View, FViewport* Viewport, FPri
 		FVector PaintRayStart = FVector::ZeroVector;
 		FVector PaintRayDirection = FVector::ZeroVector;
 
-		const int32 FirstHandIndex = PaintingWithHandIndexInVR != INDEX_NONE ? PaintingWithHandIndexInVR : 0;
-		const int32 LastHandIndex = PaintingWithHandIndexInVR != INDEX_NONE ? PaintingWithHandIndexInVR : 1;
-
-		for( int32 HandIndex = FirstHandIndex; HandIndex <= LastHandIndex; ++HandIndex )
+		// Check to see if VREditorMode is active
+		if ( IVREditorModule::IsAvailable() )
 		{
-			// Check to see if VREditorMode is active
-			if( IVREditorModule::IsAvailable() )
+			IVREditorMode* VREditorMode = static_cast< IVREditorMode* >( GetModeManager()->GetActiveMode( IVREditorModule::Get().GetVREditorModeID() ) );
+			if ( VREditorMode != nullptr && VREditorMode->IsFullyInitialized() )
 			{
-				IVREditorMode* VREditorMode = static_cast<IVREditorMode*>( GetModeManager()->GetActiveMode( IVREditorModule::Get().GetVREditorModeID() ) );
-				if( VREditorMode != nullptr )
+				for ( UViewportInteractor* Interactor : VREditorMode->GetWorldInteraction().GetInteractors() )
 				{
+					// Skip other interactors if we are painting with one
+					if ( PaintingWithInteractorInVR && Interactor != PaintingWithInteractorInVR )
+					{
+						continue;
+					}
+
 					FVector LaserPointerStart, LaserPointerEnd;
-					if( VREditorMode->GetLaserPointer( HandIndex, /* Out */ LaserPointerStart, /* Out */ LaserPointerEnd ) )
+					if ( Interactor->GetLaserPointer( /* Out */ LaserPointerStart, /* Out */ LaserPointerEnd ) )
 					{
 						const FVector LaserPointerDirection = ( LaserPointerEnd - LaserPointerStart ).GetSafeNormal();
 
@@ -2638,71 +2646,70 @@ void FEdModeMeshPaint::Render( const FSceneView* View, FViewport* Viewport, FPri
 						PaintRayStart = LaserPointerStart;
 						PaintRayDirection = LaserPointerDirection;
 					}
-				}
-				else
-				{
-					// VREditor mode not currently active
-				}
-			}
-
 
-			if( !bHavePaintRay )
-			{
-				// Make sure the cursor is visible OR we're flood filling.  No point drawing a paint cue when there's no cursor.
-				if( Viewport->IsCursorVisible() || IsForceRendered())
-				{
-					if( !PDI->IsHitTesting() )
+					if ( !bHavePaintRay )
 					{
-						// Grab the mouse cursor position
-						FIntPoint MousePosition;
-						Viewport->GetMousePos( MousePosition );
-
-						// Is the mouse currently over the viewport? or flood filling
-						if(IsForceRendered() || ( MousePosition.X >= 0 && MousePosition.Y >= 0 && MousePosition.X < (int32)Viewport->GetSizeXY().X && MousePosition.Y < (int32)Viewport->GetSizeXY().Y) )
+						// Make sure the cursor is visible OR we're flood filling.  No point drawing a paint cue when there's no cursor.
+						if ( Viewport->IsCursorVisible() || IsForceRendered() )
 						{
-							// Compute a world space ray from the screen space mouse coordinates
-							FViewportCursorLocation MouseViewportRay( View, ViewportClient, MousePosition.X, MousePosition.Y );
+							if ( !PDI->IsHitTesting() )
+							{
+								// Grab the mouse cursor position
+								FIntPoint MousePosition;
+								Viewport->GetMousePos( MousePosition );
+
+								// Is the mouse currently over the viewport? or flood filling
+								if ( IsForceRendered() || ( MousePosition.X >= 0 && MousePosition.Y >= 0 && MousePosition.X < ( int32 ) Viewport->GetSizeXY().X && MousePosition.Y < ( int32 ) Viewport->GetSizeXY().Y ) )
+								{
+									// Compute a world space ray from the screen space mouse coordinates
+									FViewportCursorLocation MouseViewportRay( View, ViewportClient, MousePosition.X, MousePosition.Y );
 
-							bHavePaintRay = true;
-							PaintRayViewOrigin = View->ViewMatrices.ViewOrigin;
-							PaintRayStart = MouseViewportRay.GetOrigin();
-							PaintRayDirection = MouseViewportRay.GetDirection();
+									bHavePaintRay = true;
+									PaintRayViewOrigin = View->ViewMatrices.ViewOrigin;
+									PaintRayStart = MouseViewportRay.GetOrigin();
+									PaintRayDirection = MouseViewportRay.GetDirection();
+								}
+							}
 						}
 					}
-				}
-			}
-
-			if( bHavePaintRay )
-			{
-				// Unless "Flow" mode is enabled, we'll only draw a visual cue while rendering and won't
-				// do any actual painting.  When "Flow" is turned on we will paint here, too!
-				const bool bVisualCueOnly = !bIsPainting || ( PaintingWithHandIndexInVR == INDEX_NONE && !FMeshPaintSettings::Get().bEnableFlow );
-				float StrengthScale = ( PaintingWithHandIndexInVR == INDEX_NONE && FMeshPaintSettings::Get().bEnableFlow ) ? FMeshPaintSettings::Get().FlowAmount : 1.0f;
 
-				// Apply VR controller trigger pressure if we're painting in VR
-				if( bIsPainting && PaintingWithHandIndexInVR != INDEX_NONE )
-				{
-					IVREditorMode* VREditorMode = static_cast<IVREditorMode*>( GetModeManager()->GetActiveMode( IVREditorModule::Get().GetVREditorModeID() ) );
-					if( VREditorMode != nullptr )
+					if ( bHavePaintRay )
 					{
-						StrengthScale *= VREditorMode->GetHandInfo( PaintingWithHandIndexInVR ).SelectAndMoveTriggerValue;
+						// Unless "Flow" mode is enabled, we'll only draw a visual cue while rendering and won't
+						// do any actual painting.  When "Flow" is turned on we will paint here, too!
+						const bool bVisualCueOnly = !bIsPainting || ( PaintingWithInteractorInVR == nullptr && !FMeshPaintSettings::Get().bEnableFlow );
+						float StrengthScale = ( PaintingWithInteractorInVR == nullptr && FMeshPaintSettings::Get().bEnableFlow ) ? FMeshPaintSettings::Get().FlowAmount : 1.0f;
 
-						// Make sure light press locking is enabled in VR for the hand we're painting with.  We don't want
-						// a full press to be converted to a move action, however locking will have been disabled already
-						// by the world interaction code after a light press happened on an object. 
-						VREditorMode->SetAllowTriggerLightPressLocking( PaintingWithHandIndexInVR, true );
-					}
-				}
+						// Apply VR controller trigger pressure if we're painting in VR
+						if ( bIsPainting && PaintingWithInteractorInVR )
+						{
+							IVREditorMode* VREditorMode = static_cast< IVREditorMode* >( GetModeManager()->GetActiveMode( IVREditorModule::Get().GetVREditorModeID() ) );
+							if ( VREditorMode != nullptr && PaintingWithInteractorInVR )
+							{
+								UVREditorInteractor* VRInteractor = Cast<UVREditorInteractor>( PaintingWithInteractorInVR );
+								if ( VRInteractor )
+								{
+									StrengthScale *= VRInteractor->GetSelectAndMoveTriggerValue();
 
-				// Apply stylus pressure if it's active
-				else if( Viewport->IsPenActive() )
-				{
-					StrengthScale *= Viewport->GetTabletPressure();
+									// Make sure light press locking is enabled in VR for the hand we're painting with.  We don't want
+									// a full press to be converted to a move action, however locking will have been disabled already
+									// by the world interaction code after a light press happened on an object. 
+									VRInteractor->SetAllowTriggerLightPressLocking( true );
+								}
+							}
+						}
+
+						// Apply stylus pressure if it's active
+						else if ( Viewport->IsPenActive() )
+						{
+							StrengthScale *= Viewport->GetTabletPressure();
+						}
+
+						const EMeshPaintAction::Type PaintAction = GetPaintAction( Viewport );
+						bool bAnyPaintableActorsUnderCursor = false;
+						DoPaint( PaintRayViewOrigin, PaintRayStart, PaintRayDirection, PDI, PaintAction, bVisualCueOnly, StrengthScale, /* Out */ bAnyPaintableActorsUnderCursor );
+					}
 				}
-				
-				const EMeshPaintAction::Type PaintAction = GetPaintAction( Viewport );
-				bool bAnyPaintableActorsUnderCursor = false;
-				DoPaint( PaintRayViewOrigin, PaintRayStart, PaintRayDirection, PDI, PaintAction, bVisualCueOnly, StrengthScale, /* Out */ bAnyPaintableActorsUnderCursor );
 			}
 		}
 	}
@@ -3497,12 +3504,13 @@ EMeshPaintAction::Type FEdModeMeshPaint::GetPaintAction(FViewport* InViewport)
 		PaintAction = bShiftDown ? EMeshPaintAction::Erase : EMeshPaintAction::Paint;
 
 		// When painting using VR, allow the modifier button to activate Erase mode
-		if( PaintingWithHandIndexInVR != INDEX_NONE )
+		if( PaintingWithInteractorInVR )
 		{
 			IVREditorMode* VREditorMode = static_cast<IVREditorMode*>( GetModeManager()->GetActiveMode( IVREditorModule::Get().GetVREditorModeID() ) );
-			if( VREditorMode != nullptr )
+			UVREditorInteractor* VRInteractor = Cast<UVREditorInteractor>( PaintingWithInteractorInVR );
+			if( VREditorMode != nullptr && VRInteractor != nullptr )
 			{
-				const bool bIsModifierPressed = VREditorMode->GetHandInfo( PaintingWithHandIndexInVR ).bIsModifierPressed;
+				const bool bIsModifierPressed = VRInteractor->IsModifierPressed();
 				PaintAction = bIsModifierPressed ? EMeshPaintAction::Erase : EMeshPaintAction::Paint;
 			}
 		}
@@ -5055,41 +5063,46 @@ TWeakObjectPtr<AActor> FEdModeMeshPaint::GetEditingActor() const
 }
 
 
-void FEdModeMeshPaint::OnVRAction( FEditorViewportClient& ViewportClient, const FVRAction VRAction, const EInputEvent Event, bool& bIsInputCaptured, bool& bWasHandled )
+void FEdModeMeshPaint::OnVRAction( class FEditorViewportClient& ViewportClient, UViewportInteractor* Interactor, 
+	const FViewportActionKeyInput& Action, const EInputEvent Event, bool& bOutIsInCaptured, bool& bWasHandled )
 {
 	IVREditorMode* VREditorMode = static_cast<IVREditorMode*>( GetModeManager()->GetActiveMode( IVREditorModule::Get().GetVREditorModeID() ) );
-	if( VREditorMode != nullptr )
+	UVREditorInteractor* VRInteractor = Cast<UVREditorInteractor>( Interactor );
+
+	if( VREditorMode != nullptr && VRInteractor != nullptr )
 	{
 		// Absorb all "full press" events unless the user is clicking on UI to avoid objects getting moved around.
 		// @todo vreditor: We shouldn't have to use this trick.  This input system needs more flexibility in its public interface!
-		if( VRAction.ActionType == EVRActionType::SelectAndMove && bIsPainting && PaintingWithHandIndexInVR == VRAction.HandIndex && !VREditorMode->GetHandInfo( VRAction.HandIndex ).bIsHoveringOverUI )
+		if( Action.ActionType == ViewportWorldActionTypes::SelectAndMove && bIsPainting && 
+			PaintingWithInteractorInVR && PaintingWithInteractorInVR == Interactor && 
+			!VRInteractor->IsHoveringOverUI() )
 		{
 			bWasHandled = true;
-			bIsInputCaptured = true;
+			bOutIsInCaptured = true;
 		}
 
-		if( VRAction.ActionType == EVRActionType::SelectAndMove_LightlyPressed )
+		if( Action.ActionType == ViewportWorldActionTypes::SelectAndMove_LightlyPressed )
 		{
-			if( !bIsPainting && Event == IE_Pressed && !VREditorMode->GetHandInfo( VRAction.HandIndex ).bIsHoveringOverUI )
+			if( !bIsPainting && Event == IE_Pressed && !VRInteractor->IsHoveringOverUI() )
 			{
 				// NOTE: We intentionally don't capture input for "LightlyPressed" because we want to allow objects to still
 				// become selected when clicked on.  We avoid allowing the objects to be moved by calling the 
 				// SetAllowTriggerLightPressLocking() function in the next tick (after VR Editor will have disabled locking.)
 				// @todo vreditor: We shouldn't have to use this trick.  This input system needs more flexibility in its public interface!
-				// 				bWasHandled = true;
-				// 				bIsInputCaptured = true;
+				//bWasHandled = true;
+				//bOutIsInCaptured = true;
 
 				StartPainting();
-				PaintingWithHandIndexInVR = VRAction.HandIndex;
+				PaintingWithInteractorInVR = Interactor;
 
 				// Go ahead and paint immediately
 				FVector LaserPointerStart, LaserPointerEnd;
-				if( VREditorMode->GetLaserPointer( PaintingWithHandIndexInVR, /* Out */ LaserPointerStart, /* Out */ LaserPointerEnd ) )
+				if( Interactor->GetLaserPointer( /* Out */ LaserPointerStart, /* Out */ LaserPointerEnd ) )
 				{
 					const FVector LaserPointerDirection = ( LaserPointerEnd - LaserPointerStart ).GetSafeNormal();
 
 					// Apply VR controller trigger pressure
-					float StrengthScale = VREditorMode->GetHandInfo( PaintingWithHandIndexInVR ).SelectAndMoveTriggerValue;
+					float StrengthScale = VRInteractor->GetSelectAndMoveTriggerValue();
 
 					// Paint!
 					const bool bVisualCueOnly = false;
@@ -5100,13 +5113,13 @@ void FEdModeMeshPaint::OnVRAction( FEditorViewportClient& ViewportClient, const
 			}
 
 			// Stop current tracking if the user is no longer painting
-			else if( bIsPainting && Event == IE_Released && PaintingWithHandIndexInVR == VRAction.HandIndex )
+			else if( bIsPainting && Event == IE_Released && PaintingWithInteractorInVR && PaintingWithInteractorInVR == Interactor )
 			{
 				EndPainting();
 
 				// @todo vreditor: We shouldn't have to use this trick.  This input system needs more flexibility in its public interface!
-				// 				bWasHandled = true;
-				// 				bIsInputCaptured = false;
+				//bWasHandled = true;
+				//bOutIsInCaptured = false;
 			}
 		}
 	}
diff --git a/Engine/Source/Editor/MeshPaint/Public/MeshPaintEdMode.h b/Engine/Source/Editor/MeshPaint/Public/MeshPaintEdMode.h
--- a/Engine/Source/Editor/MeshPaint/Public/MeshPaintEdMode.h
+++ b/Engine/Source/Editor/MeshPaint/Public/MeshPaintEdMode.h
@@ -715,7 +715,8 @@ private:
 	void OnPreReimportAsset(UObject* Object);
 
 	/** Called when the user presses a button on their motion controller device */
-	void OnVRAction( FEditorViewportClient& ViewportClient, const FVRAction VRAction, const EInputEvent Event, bool& bIsInputCaptured, bool& bOutWasHandled ); 
+	void OnVRAction( class FEditorViewportClient& ViewportClient, class UViewportInteractor* Interactor,
+		const struct FViewportActionKeyInput& Action, const EInputEvent Event, bool& bOutIsInCaptured, bool& bWasHandled );
 
 private:
 
@@ -723,7 +724,7 @@ private:
 	bool bIsPainting;
 
 	/** When painting in VR, this is the hand index that we're painting with.  Otherwise INDEX_NONE. */
-	int32 PaintingWithHandIndexInVR;
+	class UViewportInteractor* PaintingWithInteractorInVR;
 
 	/** Whether we are doing a flood fill the next render */
 	bool bIsFloodFill;
diff --git a/Engine/Source/Editor/VREditor/Gizmo/VREditorPlaneTranslationGizmoHandle.cpp b/Engine/Source/Editor/VREditor/Gizmo/VREditorPlaneTranslationGizmoHandle.cpp
--- a/Engine/Source/Editor/VREditor/Gizmo/VREditorPlaneTranslationGizmoHandle.cpp
+++ b/Engine/Source/Editor/VREditor/Gizmo/VREditorPlaneTranslationGizmoHandle.cpp
@@ -2,7 +2,7 @@
 
 #include "VREditorModule.h"
 #include "VREditorPlaneTranslationGizmoHandle.h"
-#include "VREditorBaseTransformGizmo.h"
+#include "VREditorTransformGizmo.h"
 #include "UnitConversion.h"
 
 UVREditorPlaneTranslationGizmoHandleGroup::UVREditorPlaneTranslationGizmoHandleGroup()
@@ -55,7 +55,7 @@ void UVREditorPlaneTranslationGizmoHandleGroup::UpdateGizmoHandleGroup(const FTr
 
 	for (int32 HandleIndex = 0; HandleIndex < Handles.Num(); ++HandleIndex)
 	{
-		FVREditorGizmoHandle& Handle = Handles[ HandleIndex ];
+		FGizmoHandle& Handle = Handles[ HandleIndex ];
 		UStaticMeshComponent* PlaneTranslationHandle = Handle.HandleMesh;
 		if (PlaneTranslationHandle != nullptr)	// Can be null if no handle for this specific placement
 		{
diff --git a/Engine/Source/Editor/VREditor/Gizmo/VREditorPlaneTranslationGizmoHandle.h b/Engine/Source/Editor/VREditor/Gizmo/VREditorPlaneTranslationGizmoHandle.h
--- a/Engine/Source/Editor/VREditor/Gizmo/VREditorPlaneTranslationGizmoHandle.h
+++ b/Engine/Source/Editor/VREditor/Gizmo/VREditorPlaneTranslationGizmoHandle.h
@@ -2,14 +2,14 @@
 
 #pragma once
 
-#include "VREditorGizmoHandle.h"
+#include "VIGizmoHandle.h"
 #include "VREditorPlaneTranslationGizmoHandle.generated.h"
 
 /**
  * Gizmo handle for translating on a plane
  */
 UCLASS()
-class UVREditorPlaneTranslationGizmoHandleGroup : public UVREditorGizmoHandleGroup
+class UVREditorPlaneTranslationGizmoHandleGroup : public UGizmoHandleGroup
 {
 	GENERATED_BODY()
 
diff --git a/Engine/Source/Editor/VREditor/Gizmo/VREditorRotationGizmoHandle.cpp b/Engine/Source/Editor/VREditor/Gizmo/VREditorRotationGizmoHandle.cpp
--- a/Engine/Source/Editor/VREditor/Gizmo/VREditorRotationGizmoHandle.cpp
+++ b/Engine/Source/Editor/VREditor/Gizmo/VREditorRotationGizmoHandle.cpp
@@ -56,7 +56,7 @@ void UVREditorRotationGizmoHandleGroup::UpdateGizmoHandleGroup( const FTransform
 	// Place the rotation handles
 	for (int32 HandleIndex = 0; HandleIndex < Handles.Num(); ++HandleIndex)
 	{
-		FVREditorGizmoHandle& Handle = Handles[ HandleIndex ];
+		FGizmoHandle& Handle = Handles[ HandleIndex ];
 
 		UStaticMeshComponent* RotationHandle = Handle.HandleMesh;
 		if (RotationHandle != nullptr)	// Can be null if no handle for this specific placement
diff --git a/Engine/Source/Editor/VREditor/Gizmo/VREditorRotationGizmoHandle.h b/Engine/Source/Editor/VREditor/Gizmo/VREditorRotationGizmoHandle.h
--- a/Engine/Source/Editor/VREditor/Gizmo/VREditorRotationGizmoHandle.h
+++ b/Engine/Source/Editor/VREditor/Gizmo/VREditorRotationGizmoHandle.h
@@ -2,14 +2,14 @@
 
 #pragma once
 
-#include "VREditorGizmoHandle.h"
+#include "VIGizmoHandle.h"
 #include "VREditorRotationGizmoHandle.generated.h"
 
 /**
  * Gizmo handle for rotation
  */
 UCLASS()
-class UVREditorRotationGizmoHandleGroup : public UVREditorGizmoHandleGroup
+class UVREditorRotationGizmoHandleGroup : public UGizmoHandleGroup
 {
 	GENERATED_BODY()
 
diff --git a/Engine/Source/Editor/VREditor/Gizmo/VREditorTranslationGizmoHandle.cpp b/Engine/Source/Editor/VREditor/Gizmo/VREditorTranslationGizmoHandle.cpp
--- a/Engine/Source/Editor/VREditor/Gizmo/VREditorTranslationGizmoHandle.cpp
+++ b/Engine/Source/Editor/VREditor/Gizmo/VREditorTranslationGizmoHandle.cpp
@@ -56,7 +56,7 @@ void UVREditorTranslationGizmoHandleGroup::UpdateGizmoHandleGroup( const FTransf
 	// Place the translation handles
 	for (int32 HandleIndex = 0; HandleIndex < Handles.Num(); ++HandleIndex)
 	{
-		FVREditorGizmoHandle& Handle = Handles[ HandleIndex ];
+		FGizmoHandle& Handle = Handles[ HandleIndex ];
 
 		UStaticMeshComponent* TranslationHandle = Handle.HandleMesh;
 		if (TranslationHandle != nullptr)	// Can be null if no handle for this specific placement
diff --git a/Engine/Source/Editor/VREditor/Gizmo/VREditorTranslationGizmoHandle.h b/Engine/Source/Editor/VREditor/Gizmo/VREditorTranslationGizmoHandle.h
--- a/Engine/Source/Editor/VREditor/Gizmo/VREditorTranslationGizmoHandle.h
+++ b/Engine/Source/Editor/VREditor/Gizmo/VREditorTranslationGizmoHandle.h
@@ -2,14 +2,14 @@
 
 #pragma once
 
-#include "VREditorGizmoHandle.h"
+#include "VIGizmoHandle.h"
 #include "VREditorTranslationGizmoHandle.generated.h"
 
 /**
  * Gizmo handle for translating
  */
 UCLASS()
-class UVREditorTranslationGizmoHandleGroup : public UVREditorGizmoHandleGroup
+class UVREditorTranslationGizmoHandleGroup : public UGizmoHandleGroup
 {
 	GENERATED_BODY()
 
diff --git a/Engine/Source/Editor/VREditor/Interactables/VREditorButton.cpp b/Engine/Source/Editor/VREditor/Interactables/VREditorButton.cpp
new file mode 100644
--- /dev/null
+++ b/Engine/Source/Editor/VREditor/Interactables/VREditorButton.cpp
@@ -0,0 +1,66 @@
+// Copyright 1998-2015 Epic Games, Inc. All Rights Reserved.
+
+#include "VREditorModule.h"
+#include "VREditorButton.h"
+#include "ViewportInteractor.h"
+
+AVREditorButton::AVREditorButton()
+{
+	const bool bAllowGizmoLighting = false;
+
+	// Setup uniform scaling
+	UStaticMesh* ButtonMesh = nullptr;
+	{
+		static ConstructorHelpers::FObjectFinder<UStaticMesh> ObjectFinder( TEXT( "/Engine/VREditor/TransformGizmo/UniformScaleHandle" ) );
+		ButtonMesh = ObjectFinder.Object;
+		check( ButtonMesh != nullptr );
+	}
+
+	ButtonMeshComponent = CreateDefaultSubobject<UStaticMeshComponent>( TEXT( "ButtonMeshComponent" ) );
+	check( ButtonMeshComponent != nullptr );
+
+	ButtonMeshComponent->SetStaticMesh( ButtonMesh );
+	ButtonMeshComponent->SetMobility( EComponentMobility::Movable );
+
+	ButtonMeshComponent->SetCollisionEnabled( ECollisionEnabled::QueryOnly );
+	ButtonMeshComponent->SetCollisionResponseToAllChannels( ECR_Ignore );
+	ButtonMeshComponent->SetCollisionResponseToChannel( COLLISION_GIZMO, ECollisionResponse::ECR_Block );
+	ButtonMeshComponent->SetCollisionObjectType( COLLISION_GIZMO );
+
+	ButtonMeshComponent->bGenerateOverlapEvents = false;
+	ButtonMeshComponent->SetCanEverAffectNavigation( false );
+	ButtonMeshComponent->bCastDynamicShadow = bAllowGizmoLighting;
+	ButtonMeshComponent->bCastStaticShadow = false;
+	ButtonMeshComponent->bAffectDistanceFieldLighting = bAllowGizmoLighting;
+	ButtonMeshComponent->bAffectDynamicIndirectLighting = bAllowGizmoLighting;
+}
+
+void AVREditorButton::OnPressed( UViewportInteractor* Interactor, const FHitResult& InHitResult, bool& bOutResultedInDrag  )
+{
+	GLog->Logf( TEXT( "OnPressed" ) );
+}
+
+void AVREditorButton::OnHover( UViewportInteractor* Interactor )
+{
+	GLog->Logf( TEXT( "OnHover" ) );
+}
+
+void AVREditorButton::OnHoverEnter( UViewportInteractor* Interactor, const FHitResult& InHitResult  )
+{
+	GLog->Logf( TEXT("OnHoverEnter") );
+}
+
+void AVREditorButton::OnHoverLeave( UViewportInteractor* Interactor, const UActorComponent* NewComponent )
+{
+	GLog->Logf( TEXT("OnHoverLeave") );
+}
+
+void AVREditorButton::OnDrag( UViewportInteractor* Interactor )
+{
+	GLog->Logf( TEXT( "OnDrag" ) );
+}
+
+void AVREditorButton::OnDragRelease( UViewportInteractor* Interactor )
+{
+	GLog->Logf( TEXT( "OnDragRelease" ) );
+}
diff --git a/Engine/Source/Editor/VREditor/Interactables/VREditorButton.h b/Engine/Source/Editor/VREditor/Interactables/VREditorButton.h
new file mode 100644
--- /dev/null
+++ b/Engine/Source/Editor/VREditor/Interactables/VREditorButton.h
@@ -0,0 +1,41 @@
+// Copyright 1998-2015 Epic Games, Inc. All Rights Reserved.
+
+#pragma once
+
+#include "VREditorBaseActor.h"
+#include "ViewportInteractableInterface.h"
+#include "VREditorButton.generated.h"
+
+//Forward declarations
+class UViewportInteractor;
+
+/**
+* A button for VR Editor
+*/
+UCLASS()
+class AVREditorButton : public AVREditorBaseActor, public IViewportInteractableInterface
+{
+	GENERATED_BODY()
+
+public:
+
+	/** Default constructor which sets up safe defaults */
+	AVREditorButton();
+
+	// Begin IViewportInteractableInterface
+	virtual void OnPressed( UViewportInteractor* Interactor, const FHitResult& InHitResult, bool& bOutResultedInDrag  ) override;
+	virtual void OnHover( UViewportInteractor* Interactor ) override;
+	virtual void OnHoverEnter( UViewportInteractor* Interactor, const FHitResult& InHitResult ) override;
+	virtual void OnHoverLeave( UViewportInteractor* Interactor, const UActorComponent* NewComponent ) override;
+	virtual void OnDrag( UViewportInteractor* Interactor ) override;
+	virtual void OnDragRelease( UViewportInteractor* Interactor ) override;
+	virtual class UViewportDragOperationComponent* GetDragOperationComponent() override { return nullptr; };
+	// End IViewportInteractableInterface
+
+private:
+
+	/** Mesh for the button to interact with */
+	UPROPERTY()
+	UStaticMeshComponent* ButtonMeshComponent;
+
+};
\ No newline at end of file
diff --git a/Engine/Source/Editor/VREditor/Public/IVREditorMode.h b/Engine/Source/Editor/VREditor/Public/IVREditorMode.h
--- a/Engine/Source/Editor/VREditor/Public/IVREditorMode.h
+++ b/Engine/Source/Editor/VREditor/Public/IVREditorMode.h
@@ -2,57 +2,18 @@
 
 #pragma once
 
-
-namespace VREditorConstants
-{
-	/** The number of virtual hands we support.  Always two. */
-	const int32 NumVirtualHands = 2;
-
-	/** Index of the left hand */
-	const int32 LeftHandIndex = 0;
-
-	/** Index of the right hand */
-	const int32 RightHandIndex = 1;
-};
-
-
 /**
  * Types of actions that can be performed with VR controller devices
  */
-enum class EVRActionType
+namespace VRActionTypes
 {
-	NoAction,
-	Undo,
-	Redo,
-	Touch,
-	Modifier,
-	WorldMovement,
-	SelectAndMove,
-	SelectAndMove_LightlyPressed,
-	Delete,
-	ConfirmRadialSelection,
-	TrackpadPositionX,
-	TrackpadPositionY,
-	TriggerAxis,
-	// ...
-	TotalActionTypes
-};
-
-
-/**
- * Represents a generic action bound to a specific VR "hand"
- */
-struct FVRAction
-{
-	FVRAction( EVRActionType InActionType, int32 InHandIndex )
-		: ActionType( InActionType )
-		, HandIndex( InHandIndex )
-	{}
-
-	EVRActionType ActionType;
-	int32 HandIndex;
-};
-
+	static const FName Touch( "Touch" );
+	static const FName Modifier( "Modifier" );
+	static const FName ConfirmRadialSelection( "ConfirmRadialSelection" );
+	static const FName TrackpadPositionX( "TrackpadPositionX" );
+	static const FName TrackpadPositionY( "TrackpadPositionY" );
+	static const FName TriggerAxis( "TriggerAxis" );
+}
 
 /**
  * VR Editor Mode public interface
@@ -62,19 +23,6 @@ class IVREditorMode : public FEdMode
 
 public:
 
-	/**
-	 * Gets the start and end point of the laser pointer for the specified hand
-	 *
-	 * @param HandIndex				Index of the hand to use
-	 * @param LasertPointerStart	(Out) The start location of the laser pointer in world space
-	 * @param LasertPointerEnd		(Out) The end location of the laser pointer in world space
-	 * @param bEvenIfUIIsInFront	If true, returns a laser pointer even if the hand has UI in front of it (defaults to false)
-	 * @param LaserLengthOverride	If zero the default laser length (VREdMode::GetLaserLength) is used
-	 *
-	 * @return	True if we have motion controller data for this hand and could return a valid result
-	 */
-	virtual bool GetLaserPointer( const int32 HandIndex, FVector& LaserPointerStart, FVector& LaserPointerEnd, const bool bEvenIfUIIsInFront = false, const float LaserLengthOverride = 0.0f ) const = 0;
-
 	/**
 	 * Gets the world space transform of the HMD (head)
 	 *
@@ -89,61 +37,31 @@ public:
 	 */
 	virtual class AActor* GetAvatarMeshActor() = 0;
 
-	/**
-	 * Checks to see if input is captured for the specified action
-	 *
-	 * @param	VRAction	The VR action that might be captured (the action type, and hand)
-	 *
-	 * @return	True if input is currently captured for this action
-	 */
-	virtual bool IsInputCaptured( const FVRAction VRAction ) const = 0;
-
-	/** Event that's triggered when the user presses a button on their motion controller device.  Check bIsInputCaptured first, 
-	    to see whether you should actually handle input.  Set bIsInputCaptured to true if you want to capture input, or clear it
-	    when finished capturing, if you were the captor.  Set bWasHandled to true if you reacted to the event.  Remember to check
-		bWasHandled yourself first, as there may be sibling subscribers that handle the event before you! */
-	DECLARE_EVENT_FiveParams( IVREditorMode, FOnVRAction, class FEditorViewportClient& /* ViewportClient */, const FVRAction /* VRAction */, const EInputEvent /* Event */, bool& /* bIsInputCaptured */, bool& /* bWasHandled */ );
-	virtual FOnVRAction& OnVRAction() = 0;
+	/** Gets the world currently used by the worldinteraction of this mode */
+	virtual UWorld* GetWorld() const = 0;
 
-	/** Event that fires every frame to update hover state based on what's under the cursor.  Set bWasHandled to true if you detect 
-	    something to hover, and set HoverImpactPoint to the world space location you hit.  Remember to check bWasHandled yourself 
-		first, as there may be sibling subscribers that handle the event before you! */
-	DECLARE_EVENT_FiveParams( IVREditorMode, FOnVRHoverUpdate, class FEditorViewportClient& /* ViewportClient */, const int32 /* HandIndex */, FVector& /* HoverImpactPoint */, bool& /* bIsHoveringOverUI */, bool& /* bWasHandled */ );
-	virtual FOnVRHoverUpdate& OnVRHoverUpdate() = 0;
+	/** Gets the world space transform of the calibrated VR room origin.  When using a seated VR device, this will feel like the
+	camera's world transform (before any HMD positional or rotation adjustments are applied.) */
+	virtual FTransform GetRoomTransform() const = 0;
 
+	/** Sets a new transform for the room, in world space.  This is basically setting the editor's camera transform for the viewport */
+	virtual void SetRoomTransform( const FTransform& NewRoomTransform ) = 0;
 
-	struct FHandInfo
-	{
-		/** Is the Modifier button held down? */
-		bool bIsModifierPressed;
+	/** Gets the transform of the user's HMD in room space */
+	virtual FTransform GetRoomSpaceHeadTransform() const = 0;
 
-		/** How much the trigger is held down for the 'select and move' action */
-		float SelectAndMoveTriggerValue;
+	/** Gets access to the world interaction system (const) */
+	virtual const class UViewportWorldInteraction& GetWorldInteraction() const = 0;
 
-		/** True if we're hovering over UI right now.  When hovering over UI, we don't bother drawing a see-thru laser pointer */
-		bool bIsHoveringOverUI;
-	};
+	/** Gets access to the world interaction system */
+	virtual class UViewportWorldInteraction& GetWorldInteraction() = 0;
 
+	/** If the mode was completely initialized */
+	virtual bool IsFullyInitialized() const = 0;
 
-	/**
-	 * Gets some information about the specified hand (motion controller)
-	 *
-	 * @param	HandIndex	Which hand to query info for
-	 *
-	 * @return	A structure containing information about the hand's current state
-	 */
-	virtual FHandInfo GetHandInfo( const int32 HandIndex ) const = 0;
-
-	/**
-	 * Enables or disable "light press" locking for the specified hand.  This is an advanced feature.
-	 * It should be true if we allow locking of the lightly pressed state for the current event.  This can 
-	 * be useful to turn off in cases where you always want a full press to override the light press, even 
-	 * if it was a very slow press
-	 *
-	 * @param	HandIndex	The hand to change light press locking state for
-	 * @param	bAllow		Enables or disables light press locking for this hand
-	 */
-	virtual void SetAllowTriggerLightPressLocking( const int32 HandIndex, const bool bAllow ) = 0;
+	/** Gets the event for update which is broadcasted by the mode itself after getting latest input states */
+	DECLARE_EVENT_OneParam( IVREditorMode, FOnVRTickHandle, const float /* DeltaTime */ );
+	virtual FOnVRTickHandle& OnTickHandle() = 0;
 
 };
 
diff --git a/Engine/Source/Editor/VREditor/Public/VREditorInteractor.h b/Engine/Source/Editor/VREditor/Public/VREditorInteractor.h
new file mode 100644
--- /dev/null
+++ b/Engine/Source/Editor/VREditor/Public/VREditorInteractor.h
@@ -0,0 +1,221 @@
+// Copyright 1998-2015 Epic Games, Inc. All Rights Reserved.
+
+#pragma once
+
+#include "Editor/ViewportInteraction/Public/ViewportInteractorMotionController.h"
+#include "VREditorInteractor.generated.h"
+
+/**
+ * VREditor default interactor
+ */
+UCLASS()
+class VREDITOR_API UVREditorInteractor : public UViewportInteractorMotionController
+{
+	GENERATED_UCLASS_BODY()
+
+public:
+
+	virtual ~UVREditorInteractor();
+
+	/** Initialize default values */
+	void Init( const EControllerHand InHandSide, class FVREditorMode* InVRMode );
+
+	/** Sets up all componentes defined in UViewportInteractorMotionController with the passed Actor as owner */
+	void SetupComponent( AActor* OwningActor );
+
+	/** Gets the owner of this system */
+	class FVREditorMode& GetVRMode()
+	{
+		return *VRMode;
+	}
+
+	/** Gets the owner of this system (const) */
+	const class FVREditorMode& GetVRMode() const
+	{
+		return *VRMode;
+	}
+
+	// ViewportInteractorInterface overrides
+	virtual void Shutdown() override;
+	virtual void Tick( const float DeltaTime ) override;
+	virtual FHitResult GetHitResultFromLaserPointer( TArray<AActor*>* OptionalListOfIgnoredActors = nullptr, const bool bIgnoreGizmos = false,
+		TArray<UClass*>* ObjectsInFrontOfGizmo = nullptr, const bool bEvenIfBlocked = false, const float LaserLengthOverride = 0.0f ) override;
+	virtual void ResetHoverState( const float DeltaTime ) override;
+	virtual void CalculateDragRay() override;
+	virtual void OnStartDragging( UActorComponent* ClickedComponent, const FVector& HitLocation, const bool bIsPlacingActors ) override;
+
+	//
+	// Getters and setters
+	//
+
+	/** Gets the Y delta of the trackpad */
+	float GetTrackpadSlideDelta();
+
+	/** Gets if*/
+	bool IsHoveringOverUI() const;
+	
+	/** Sets if there is a docked window on this interactor */
+	void SetHasUIInFront( const bool bInHasUIInFront );
+	
+	/** Check if there is any docked window on this interactor */
+	bool HasUIInFront() const;
+
+	/** Sets if the quick menu is on this interactor */
+	void SetHasUIOnForearm( const bool bInHasUIOnForearm );
+	
+	/** Check if the quick menu is on this interactor */
+	bool HasUIOnForearm() const;
+
+	/** Check if the touchpad is currently touched */
+	bool IsTouchingTrackpad() const;
+
+	/** Get the current position of the trackpad or analog stick */
+	FVector2D GetTrackpadPosition() const;
+
+	/** Get the last position of the trackpad or analog stick */
+	FVector2D GetLastTrackpadPosition() const;
+
+	/** Gets the current hovered widget component if any */
+	UWidgetComponent* GetHoveringOverWidgetComponent() const;
+
+	/** Sets the current hovered widget component */
+	void SetHoveringOverWidgetComponent( UWidgetComponent* NewHoveringOverWidgetComponent );
+	
+	/** If the trackpad values are valid */
+	bool IsTrackpadPositionValid( const int32 AxisIndex ) const;
+
+	/** Get when the last time the trackpad position was updated */
+	FTimespan& GetLastTrackpadPositionUpdateTime();
+
+	/** If the modifier key is currently pressed */
+	bool IsModifierPressed() const;
+
+	/** Sets if the interactor is clicking on any UI */
+	void SetIsClickingOnUI( const bool bInIsClickingOnUI );
+	
+	/** Gets if the interactor is clicking on any UI */
+	bool IsClickingOnUI() const;
+
+	/** Sets if the interactor is hovering over any UI */
+	void SetIsHoveringOverUI( const bool bInIsHoveringOverUI );
+
+	/** Sets if the interactor is right  hover over any UI */
+	void SetIsRightClickingOnUI( const bool bInIsRightClickingOnUI );
+
+	/** Gets if the interactor is right clicking on UI */
+	bool IsRightClickingOnUI() const;
+
+	/** Gets if the interactor is right clicking on UI */
+	void SetLastClickReleaseTime( const double InLastClickReleaseTime );
+
+	/** Gets last time the interactor released */
+	double GetLastClickReleaseTime() const;
+
+	/** Sets the UI scroll velocity */
+	void SetUIScrollVelocity( const float InUIScrollVelocity );
+
+	/** Gets the UI scroll velocity */
+	float GetUIScrollVelocity() const;
+	
+	/* ViewportInteractor overrides, checks if the laser is blocked by UI */
+	virtual bool GetIsLaserBlocked() override;
+
+protected:
+
+	// ViewportInteractor overrides
+	virtual void HandleInputKey( FViewportActionKeyInput& Action, const FKey Key, const EInputEvent Event, bool& bOutWasHandled ) override;
+	virtual void HandleInputAxis( FViewportActionKeyInput& Action, const FKey Key, const float Delta, const float DeltaTime, bool& bOutWasHandled ) override;
+
+private:
+
+	/** Changes the color of the buttons on the handmesh */
+	void ApplyButtonPressColors( const FViewportActionKeyInput& Action );
+
+	/** Set the visuals for a button on the motion controller */
+	void SetMotionControllerButtonPressedVisuals( const EInputEvent Event, const FName& ParameterName, const float PressStrength );
+	
+	/** Pops up some help text labels for the controller in the specified hand, or hides it, if requested */
+	void ShowHelpForHand( const bool bShowIt );
+	
+	/** Called every frame to update the position of any floating help labels */
+	void UpdateHelpLabels();
+
+	/** Given a mesh and a key name, tries to find a socket on the mesh that matches a supported key */
+	UStaticMeshSocket* FindMeshSocketForKey( UStaticMesh* StaticMesh, const FKey Key );
+
+	/** Sets the visuals of the LaserPointer */
+	void SetLaserVisuals( const FLinearColor& NewColor, const float CrawlFade, const float CrawlSpeed );
+
+	/** Updates the radial menu */
+	void UpdateRadialMenuInput( const float DeltaTime );
+
+private:
+	
+	/** The mode that created this interactor */
+	class FVREditorMode* VRMode;
+
+	//
+	// UI
+	//
+
+	/** True if a floating UI panel is attached to the front of the hand, and we shouldn't bother drawing a laser
+	pointer or enabling certain other features */
+	bool bHasUIInFront;
+
+	/** True if a floating UI panel is attached to our forearm, so we shouldn't bother drawing help labels */
+	bool bHasUIOnForearm;
+
+	/** True if we're currently holding the 'SelectAndMove' button down after clicking on UI */
+	bool bIsClickingOnUI;
+
+	/** When bIsClickingOnUI is true, this will be true if we're "right clicking".  That is, the Modifier key was held down at the time that the user clicked */
+	bool bIsRightClickingOnUI;
+
+	 /** True if we're hovering over UI right now.  When hovering over UI, we don't bother drawing a see-thru laser pointer */
+	bool bIsHoveringOverUI;
+	
+	/** Inertial scrolling -- how fast to scroll the mousewheel over UI */
+	float UIScrollVelocity;	
+
+	//
+	// Trackpad support
+	//
+
+	/** True if the trackpad is actively being touched */
+	bool bIsTouchingTrackpad;
+
+	/** Position of the touched trackpad */
+	FVector2D TrackpadPosition;
+
+	/** Last position of the touched trackpad */
+	FVector2D LastTrackpadPosition;
+
+	/** True if we have a valid trackpad position (for each axis) */
+	bool bIsTrackpadPositionValid[ 2 ];
+
+	/** Real time that the last trackpad position was last updated.  Used to filter out stale previous data. */
+	FTimespan LastTrackpadPositionUpdateTime;
+
+	//
+	// General input @todo: VREditor: Should this be general (non-UI) in interactordata ?
+	//
+
+	/** Is the Modifier button held down? */
+	bool bIsModifierPressed;
+
+	/** Last real time that we released the 'SelectAndMove' button on UI.  This is used to detect double-clicks. */
+	double LastClickReleaseTime;
+	
+	//
+	// Help
+	//
+
+	/** True if we want help labels to be visible right now, otherwise false */
+	bool bWantHelpLabels;
+
+	/** Help labels for buttons on the motion controllers */
+	TMap< FKey, class AFloatingText* > HelpLabels;
+
+	/** Time that we either started showing or hiding help labels (for fade transitions) */
+	FTimespan HelpLabelShowOrHideStartTime;
+};
\ No newline at end of file
diff --git a/Engine/Source/Editor/VREditor/Teleporter/VREditorTeleporter.cpp b/Engine/Source/Editor/VREditor/Teleporter/VREditorTeleporter.cpp
new file mode 100644
--- /dev/null
+++ b/Engine/Source/Editor/VREditor/Teleporter/VREditorTeleporter.cpp
@@ -0,0 +1,142 @@
+// Copyright 1998-2015 Epic Games, Inc. All Rights Reserved.
+
+#include "VREditorModule.h"
+#include "VREditorTeleporter.h"
+#include "IVREditorMode.h"
+#include "ViewportWorldInteraction.h"
+#include "VREditorInteractor.h"
+
+#include "Kismet/GameplayStatics.h"
+
+namespace VREd
+{
+	static FAutoConsoleVariable TeleportLerpTime( TEXT( "VREd.TeleportLerpTime" ), 0.1f, TEXT( "The lerp time to teleport" ) );
+	static FAutoConsoleVariable TeleportOffset( TEXT( "VREd.TeleportOffset" ), 100.0f, TEXT( "The offset from the hitresult towards the controller" ) );
+	static FAutoConsoleVariable TeleportLaserPointerLength( TEXT( "VREd.TeleportLaserPointerLength" ), 500000.0f, TEXT( "Distance of the LaserPointer for teleporting" ) );
+}
+
+UVREditorTeleporter::UVREditorTeleporter( const FObjectInitializer& Initializer ) :
+	Super( Initializer ),
+	bIsTeleporting( false ),
+	TeleportLerpAlpha( 0 ),
+	TeleportStartLocation( FVector::ZeroVector ),
+	TeleportGoalLocation( FVector::ZeroVector ),
+	bJustTeleported( false ),
+	VRMode( nullptr )
+{
+	// Load sounds
+	TeleportSound = LoadObject<USoundCue>( nullptr, TEXT( "/Engine/VREditor/Sounds/VR_teleport_Cue" ) );
+	check( TeleportSound != nullptr );
+}
+
+void UVREditorTeleporter::Init( IVREditorMode* InMode )
+{
+	VRMode = InMode;
+	VRMode->OnTickHandle().AddUObject( this, &UVREditorTeleporter::Tick );
+	VRMode->GetWorldInteraction().OnViewportInteractionInputAction().AddUObject( this, &UVREditorTeleporter::OnInteractorAction );
+}
+
+void UVREditorTeleporter::Shutdown()
+{
+	TeleportSound = nullptr;
+}
+
+void UVREditorTeleporter::Tick( const float DeltaTime )
+{
+	if ( bIsTeleporting )
+	{
+		Teleport( DeltaTime );
+
+		//if ( bJustTeleported && !bIsTeleporting )
+		//{
+		//	bJustTeleported = false;
+		//}
+	}
+}
+
+void UVREditorTeleporter::StartTeleport( UViewportInteractor* Interactor )
+{
+	UVREditorInteractor* VREditorInteractor = Cast<UVREditorInteractor>( Interactor );
+	if ( VREditorInteractor && !VREditorInteractor->IsHoveringOverUI() )
+	{
+		FVector LaserPointerStart, LaserPointerEnd;
+		if ( Interactor->GetLaserPointer( LaserPointerStart, LaserPointerEnd ) )
+		{
+			FHitResult HitResult = Interactor->GetHitResultFromLaserPointer( nullptr, true, nullptr, false, VREd::TeleportLaserPointerLength->GetFloat() );
+			if ( HitResult.bBlockingHit )
+			{
+				const FVector HeadLocation = VRMode->GetHeadTransform().GetLocation();
+				const FVector RoomLocation = VRMode->GetRoomTransform().GetLocation();
+
+				const FVector Offset = ( LaserPointerStart - HitResult.Location ).GetSafeNormal() *
+					VREd::TeleportOffset->GetFloat() * VRMode->GetWorldInteraction().GetWorldScaleFactor(); 	// Offset from the hitresult, TeleportOffset is in centimeter
+				const FVector NewHandLocation = HitResult.Location + Offset;								// New world location offset from the hitresult
+				const FVector HandHeadOffset = HeadLocation - Interactor->GetTransform().GetLocation();		// The offset between the hand and the head
+				const FVector NewHeadLocation = NewHandLocation + HandHeadOffset;							// The location in world space where the head is going to be
+																											// The new roomspace location in world space
+				const FVector HeadRoomOffset = HeadLocation - RoomLocation;									// Where the head is going to be in world space
+				const FVector NewRoomSpaceLocation = NewHeadLocation - HeadRoomOffset;						// Result, where the roomspace is going to be in world space
+				
+																											// Set values to start teleporting
+				bIsTeleporting = true;
+				TeleportLerpAlpha = 0.0f;
+				TeleportStartLocation = RoomLocation;
+				TeleportGoalLocation = NewRoomSpaceLocation;
+
+				UGameplayStatics::PlaySound2D( VRMode->GetWorld(), TeleportSound );
+			}
+		}
+	}
+}
+
+void UVREditorTeleporter::OnInteractorAction( class FEditorViewportClient& ViewportClient, UViewportInteractor* Interactor, const struct FViewportActionKeyInput& Action, const EInputEvent Event, bool& bOutIsInCaptured, bool& bWasHandled )
+{
+	if ( ( Action.ActionType == ViewportWorldActionTypes::SelectAndMove ||
+		Action.ActionType == ViewportWorldActionTypes::SelectAndMove_LightlyPressed ) )
+	{
+		if ( Event == IE_Pressed )
+		{
+			if ( Action.ActionType == ViewportWorldActionTypes::SelectAndMove_LightlyPressed )
+			{
+				// Check to teleport before checking on actor selection, otherwise objects will be selected when teleporting
+				if	( Interactor->GetDraggingMode() == EViewportInteractionDraggingMode::World &&
+					( Interactor->GetOtherInteractor() == nullptr || Interactor->GetOtherInteractor()->GetDraggingMode() != EViewportInteractionDraggingMode::AssistingDrag ) )
+				{
+					StartTeleport( Interactor );
+					bJustTeleported = true;
+					bWasHandled = true;
+				}
+			}
+			else if ( Action.ActionType == ViewportWorldActionTypes::SelectAndMove && bJustTeleported )
+			{
+				//Block the input
+				bJustTeleported = false;
+				bWasHandled = true;
+			}
+
+			// Nothing should happen when the hand has a UI in front, otherwise it will deselect everything
+			if ( Interactor->GetIsLaserBlocked() )
+			{
+				bWasHandled = true;
+			}
+		}
+	}
+}
+
+void UVREditorTeleporter::Teleport( const float DeltaTime )
+{
+	FTransform RoomTransform = VRMode->GetRoomTransform();
+	TeleportLerpAlpha += DeltaTime;
+
+	if ( TeleportLerpAlpha > VREd::TeleportLerpTime->GetFloat() )
+	{
+		// Teleporting is finished
+		TeleportLerpAlpha = VREd::TeleportLerpTime->GetFloat();
+		bIsTeleporting = false;
+	}
+
+	// Calculate the new position of the roomspace
+	FVector NewLocation = FMath::Lerp( TeleportStartLocation, TeleportGoalLocation, TeleportLerpAlpha / VREd::TeleportLerpTime->GetFloat() );
+	RoomTransform.SetLocation( NewLocation );
+	VRMode->SetRoomTransform( RoomTransform );
+}
\ No newline at end of file
diff --git a/Engine/Source/Editor/VREditor/Teleporter/VREditorTeleporter.h b/Engine/Source/Editor/VREditor/Teleporter/VREditorTeleporter.h
new file mode 100644
--- /dev/null
+++ b/Engine/Source/Editor/VREditor/Teleporter/VREditorTeleporter.h
@@ -0,0 +1,70 @@
+// Copyright 1998-2015 Epic Games, Inc. All Rights Reserved.
+
+#pragma once
+
+#include "VREditorTeleporter.generated.h"
+
+/**
+* VR Editor teleport manager
+*/
+UCLASS()
+class UVREditorTeleporter : public UObject
+{
+	GENERATED_UCLASS_BODY()
+
+public:
+
+	/** Initializes the teleport system */
+	void Init( class IVREditorMode* InMode );
+
+	/** Shuts down the teleport system */
+	void Shutdown();
+
+	/** Ticks the current teleport */
+	void Tick( const float DeltaTime );
+
+	/** Returns true if we're actively teleporting */
+	bool IsTeleporting() const
+	{
+		return bIsTeleporting;
+	}
+
+	/** Start teleporting, does a ray trace with the hand passed and calculates the locations for lerp movement in Teleport */
+	void StartTeleport( class UViewportInteractor* Interactor );
+
+private:
+
+	/** Called when the user presses a button on their motion controller device */
+	void OnInteractorAction( class FEditorViewportClient& ViewportClient, UViewportInteractor* Interactor,
+		const struct FViewportActionKeyInput& Action, const EInputEvent Event, bool& bOutIsInCaptured, bool& bWasHandled );
+
+	/** Actually teleport to a position */
+	void Teleport( const float DeltaTime );
+
+	//
+	// Fields
+	//
+
+	/** Owning mode */
+	IVREditorMode* VRMode;
+
+	/** If we are currently teleporting */
+	bool bIsTeleporting;
+
+	/** The current lerp of the teleport between the TeleportStartLocation and the TeleportGoalLocation */
+	float TeleportLerpAlpha;
+
+	/** Set on the current Roomspace location in the world in StartTeleport before doing the actual teleporting */
+	FVector TeleportStartLocation;
+
+	/** The calculated goal location in StartTeleport to move the Roomspace to */
+	FVector TeleportGoalLocation;
+
+	/** So we can teleport on a light press and make sure a second teleport by a fullpress doesn't happen */
+	bool bJustTeleported;
+
+	/** Teleport sound */
+	UPROPERTY()
+	class USoundCue* TeleportSound;
+
+};
\ No newline at end of file
diff --git a/Engine/Source/Editor/VREditor/VREditor.Build.cs b/Engine/Source/Editor/VREditor/VREditor.Build.cs
--- a/Engine/Source/Editor/VREditor/VREditor.Build.cs
+++ b/Engine/Source/Editor/VREditor/VREditor.Build.cs
@@ -18,7 +18,8 @@ namespace UnrealBuildTool.Rules
 				    "UnrealEd",
 					"UMG",
 					"LevelEditor",
-					"HeadMountedDisplay"
+					"HeadMountedDisplay",
+                    "ViewportInteraction"
 				}
 			);
 
diff --git a/Engine/Source/Editor/VREditor/VREditorDockableWindow.cpp b/Engine/Source/Editor/VREditor/VREditorDockableWindow.cpp
--- a/Engine/Source/Editor/VREditor/VREditorDockableWindow.cpp
+++ b/Engine/Source/Editor/VREditor/VREditorDockableWindow.cpp
@@ -7,6 +7,8 @@
 #include "VREditorBaseUserWidget.h"
 #include "WidgetComponent.h"
 #include "VREditorWidgetComponent.h"
+#include "VREditorWorldInteraction.h"
+#include "VREditorInteractor.h"
 
 namespace VREd
 {
@@ -18,10 +20,11 @@ namespace VREd
 	static FAutoConsoleVariable DockUIHoverAnimationDuration( TEXT( "VREd.DockUIHoverAnimationDuration" ), 0.15f, TEXT( "How quick the hover animation should complete in" ) );
 }
 
-AVREditorDockableWindow::AVREditorDockableWindow() 
-	: Super()
-{
 
+AVREditorDockableWindow::AVREditorDockableWindow()
+	: Super(),
+	DockSelectDistance( 0.0f )
+{
 	{
 		UStaticMesh* WindowMesh = nullptr;
 		{
@@ -141,9 +144,25 @@ AVREditorDockableWindow::AVREditorDockableWindow()
 	// their laser toward the UI
 	SelectionBarMeshComponent->SetVisibility( false );
 	CloseButtonMeshComponent->SetVisibility( false );
+
+	// Create the drag operation
+	DragOperationComponent = CreateDefaultSubobject<UViewportDragOperationComponent>( TEXT( "DragOperation" ) );
+	DragOperationComponent->SetDragOperactionClass( UDockableWindowDragOperation::StaticClass() );
 }
 
 
+AVREditorDockableWindow::~AVREditorDockableWindow()
+{
+	WindowMeshComponent = nullptr;
+	SelectionBarMeshComponent = nullptr;
+	CloseButtonMeshComponent = nullptr;
+	SelectionBarMID = nullptr;
+	SelectionBarTranslucentMID = nullptr;
+	CloseButtonMID = nullptr;
+	CloseButtonTranslucentMID = nullptr;
+	DragOperationComponent = nullptr;
+}
+
 void AVREditorDockableWindow::SetupWidgetComponent()
 {
 	Super::SetupWidgetComponent();
@@ -175,9 +194,9 @@ void AVREditorDockableWindow::TickManually( float DeltaTime )
 			const FVector UIForwardVector = FVector::ForwardVector;
 			const float MinDotForAimingAtOtherHand = -1.1f;	// @todo vreditor tweak
 
-			for( int32 HandIndex = 0; HandIndex < VREditorConstants::NumVirtualHands; ++HandIndex )
+			for( UViewportInteractor* Interactor : GetOwner().GetOwner().GetWorldInteraction().GetInteractors() )
 			{
-				if( GetOwner().GetOwner().IsHandAimingTowardsCapsule( HandIndex, UICapsuleTransform, UICapsuleStart, UICapsuleEnd, UICapsuleLocalRadius, MinDistanceToUICapsule, UIForwardVector, MinDotForAimingAtOtherHand ) )
+				if( GetOwner().GetOwner().IsHandAimingTowardsCapsule( Interactor, UICapsuleTransform, UICapsuleStart, UICapsuleEnd, UICapsuleLocalRadius, MinDistanceToUICapsule, UIForwardVector, MinDotForAimingAtOtherHand ) )
 				{
 					bIsLaserAimingTowardUI = true;
 				}
@@ -285,50 +304,140 @@ void AVREditorDockableWindow::UpdateRelativeRoomTransform()
 	const FVector RoomSpaceWindowLocation = WindowToRoomTransform.GetLocation() / GetOwner().GetOwner().GetWorldScaleFactor();;
 	const FQuat RoomSpaceWindowRotation = WindowToRoomTransform.GetRotation();
 
-	RelativeOffset = RoomSpaceWindowLocation;
-	LocalRotation = RoomSpaceWindowRotation.Rotator();
+	SetRelativeOffset( RoomSpaceWindowLocation );
+	SetLocalRotation( RoomSpaceWindowRotation.Rotator() );
+}
+
+UStaticMeshComponent* AVREditorDockableWindow::GetCloseButtonMeshComponent() const
+{
+	return CloseButtonMeshComponent;
+}
+
+UStaticMeshComponent* AVREditorDockableWindow::GetSelectionBarMeshComponent() const
+{
+	return SelectionBarMeshComponent;
+}
+
+float AVREditorDockableWindow::GetDockSelectDistance() const
+{
+	return DockSelectDistance;
+}
+
+void AVREditorDockableWindow::OnPressed( UViewportInteractor* Interactor, const FHitResult& InHitResult, bool& bOutResultedInDrag )
+{
+	bOutResultedInDrag = false;
+	UVREditorInteractor* VRInteractor = Cast<UVREditorInteractor>( Interactor );
+
+	if ( InHitResult.Component == GetCloseButtonMeshComponent() )
+	{
+		// Close the window
+		const bool bShouldShow = false;
+		const bool bShowOnHand = false;
+		const bool bRefreshQuickMenu = true;
+		GetOwner().ShowEditorUIPanel( this, VRInteractor, bShouldShow, bShowOnHand, bRefreshQuickMenu );
+	}
+	else if ( InHitResult.Component == GetSelectionBarMeshComponent() )
+	{
+		bOutResultedInDrag = true;
+		
+		FVector LaserPointerStart, LaserPointerEnd;
+		DockSelectDistance = ( InHitResult.TraceStart - InHitResult.Location ).Size();
+		GetOwner().StartDraggingDockUI( this, VRInteractor, DockSelectDistance );
+	}
+}
+
+void AVREditorDockableWindow::OnHover( UViewportInteractor* Interactor )
+{
+
 }
 
-void AVREditorDockableWindow::OnEnterHover( const FHitResult& HitResult, const int32 HandIndex )
+void AVREditorDockableWindow::OnHoverEnter( UViewportInteractor* Interactor, const FHitResult& InHitResult )
 {
-	if( SelectionBarMeshComponent == HitResult.GetComponent() )
+	if ( SelectionBarMeshComponent == InHitResult.GetComponent() )
 	{
 		bIsHoveringOverSelectionBar = true;
 	}
-	
-	if( CloseButtonMeshComponent == HitResult.GetComponent() )
+
+	if ( CloseButtonMeshComponent == InHitResult.GetComponent() )
 	{
 		bIsHoveringOverCloseButton = true;
 	}
 }
 
-void AVREditorDockableWindow::OnLeaveHover( const int32 HandIndex, const UActorComponent* NewComponent  )
+void AVREditorDockableWindow::OnHoverLeave( UViewportInteractor* Interactor, const UActorComponent* NewComponent )
 {
-	FVirtualHand& OtherHand = GetOwner().GetOwner().GetOtherHand( HandIndex );
+	UActorComponent* OtherInteractorHoveredComponent = Interactor->GetOtherInteractor()->GetHoverComponent();
 
-	if( OtherHand.HoveredActorComponent != SelectionBarMeshComponent && NewComponent != SelectionBarMeshComponent )
+	if ( OtherInteractorHoveredComponent != SelectionBarMeshComponent && NewComponent != SelectionBarMeshComponent )
 	{
 		bIsHoveringOverSelectionBar = false;
 	}
-	
-	if( OtherHand.HoveredActorComponent != CloseButtonMeshComponent && NewComponent != CloseButtonMeshComponent  )
+
+	if ( OtherInteractorHoveredComponent != CloseButtonMeshComponent && NewComponent != CloseButtonMeshComponent )
 	{
 		bIsHoveringOverCloseButton = false;
 	}
 }
 
+void AVREditorDockableWindow::OnDrag( UViewportInteractor* Interactor )
+{
 
-void AVREditorDockableWindow::SetSelectionBarColor( const FLinearColor& LinearColor )
+}
+
+void AVREditorDockableWindow::OnDragRelease( UViewportInteractor* Interactor )
 {
+	UVREditorInteractor* VREditorInteractor = Cast<UVREditorInteractor>( Interactor );
+	UVREditorUISystem& UISystem = GetOwner();
+	UISystem.StopDraggingDockUI( VREditorInteractor );
+}
+
+UViewportDragOperationComponent* AVREditorDockableWindow::GetDragOperationComponent()
+{
+	return DragOperationComponent;
+}
+
+void AVREditorDockableWindow::SetSelectionBarColor( const FLinearColor& LinearColor )
+{		
 	static FName StaticColorParameterName( "Color" );
 	SelectionBarMID->SetVectorParameterValue( StaticColorParameterName, LinearColor );
 	SelectionBarTranslucentMID->SetVectorParameterValue( StaticColorParameterName, LinearColor );
 }
 
-
 void AVREditorDockableWindow::SetCloseButtonColor( const FLinearColor& LinearColor )
 {
 	static FName StaticColorParameterName( "Color" );
 	CloseButtonMID->SetVectorParameterValue( StaticColorParameterName, LinearColor );
 	CloseButtonTranslucentMID->SetVectorParameterValue( StaticColorParameterName, LinearColor );
 }
+
+/************************************************************************/
+/* Dockable window drag operation                                       */
+/******************	******************************************************/
+void UDockableWindowDragOperation::ExecuteDrag( UViewportInteractor* Interactor, IViewportInteractableInterface* Interactable  )
+{
+	UVREditorInteractor* VRInteractor = Cast<UVREditorInteractor>( Interactor );
+	AVREditorDockableWindow* DockableWindow = Cast<AVREditorDockableWindow>( Interactable );
+	if ( VRInteractor && DockableWindow )
+	{
+		UVREditorUISystem& UISystem = DockableWindow->GetOwner();
+
+		const bool bIsAbsolute = ( UISystem.GetOwner().GetHMDDeviceType() == EHMDDeviceType::DT_SteamVR );
+		float SlideDeltaY = VRInteractor->GetTrackpadSlideDelta();
+
+		float NewUIScale = DockableWindow->GetScale() + SlideDeltaY;
+		if ( NewUIScale <= UISystem.GetMinDockWindowSize() )
+		{
+			NewUIScale = UISystem.GetMinDockWindowSize();
+		}
+		else if ( NewUIScale >= UISystem.GetMaxDockWindowSize() )
+		{
+			NewUIScale = UISystem.GetMaxDockWindowSize();
+		}
+
+		DockableWindow->SetScale( NewUIScale );
+
+		const FTransform DockedUIToWorld = UISystem.MakeDockableUITransform( DockableWindow, VRInteractor, DockableWindow->GetDockSelectDistance() );
+		DockableWindow->SetActorTransform( DockedUIToWorld );
+		DockableWindow->UpdateRelativeRoomTransform();
+	}
+}
diff --git a/Engine/Source/Editor/VREditor/VREditorDockableWindow.h b/Engine/Source/Editor/VREditor/VREditorDockableWindow.h
--- a/Engine/Source/Editor/VREditor/VREditorDockableWindow.h
+++ b/Engine/Source/Editor/VREditor/VREditorDockableWindow.h
@@ -3,13 +3,14 @@
 #pragma once
 
 #include "VREditorFloatingUI.h"
+#include "ViewportInteractableInterface.h"
 #include "VREditorDockableWindow.generated.h"
 
 /**
  * An interactive floating UI panel that can be dragged around
  */
 UCLASS()
-class AVREditorDockableWindow : public AVREditorFloatingUI
+class AVREditorDockableWindow : public AVREditorFloatingUI, public IViewportInteractableInterface
 {
 	GENERATED_BODY()
 	
@@ -17,36 +18,45 @@ public:
 
 	/** Default constructor */
 	AVREditorDockableWindow();
+
+	/** Destructor */
+	~AVREditorDockableWindow();
 	
 	/** Updates the meshes for the UI */
 	virtual void TickManually( float DeltaTime ) override;
 
 	/** Updates the last dragged relative position */
 	void UpdateRelativeRoomTransform();
-
-	/** Enter hover with laser changes the color of SelectionMesh and CloseButtonMesh */
-	void OnEnterHover( const FHitResult& HitResult, const int32 HandIndex );
-
-	/** Leaving hover with laser changes the color of SelectionMesh and CloseButtonMesh */
-	void OnLeaveHover( const int32 HandIndex, const class UActorComponent* NewHoveredComponent );
 	
 	/** Gets the close button component */
-	UStaticMeshComponent* GetCloseButtonMeshComponent()
-	{
-		return CloseButtonMeshComponent;
-	}
+	UStaticMeshComponent* GetCloseButtonMeshComponent() const;
 
 	/** Gets the selection bar component */
-	UStaticMeshComponent* GetSelectionBarMeshComponent()
-	{
-		return SelectionBarMeshComponent;
-	}
+	UStaticMeshComponent* GetSelectionBarMeshComponent() const;
 
-private:
+	/** Gets the distance between the interactor and the window when sta	rting drag */
+	float GetDockSelectDistance() const;
+
+	// IViewportInteractableInterface
+	virtual void OnPressed( UViewportInteractor* Interactor, const FHitResult& InHitResult, bool& bOutResultedInDrag ) override;
+	virtual void OnHover( UViewportInteractor* Interactor ) override;
+	
+	/** Enter hover with laser changes the color of SelectionMesh and CloseButtonMesh */
+	virtual void OnHoverEnter( UViewportInteractor* Interactor, const FHitResult& InHitResult ) override;
+	
+	/** Leaving hover with laser changes the color of SelectionMesh and CloseButtonMesh */
+	virtual void OnHoverLeave( UViewportInteractor* Interactor, const UActorComponent* NewComponent ) override;
+	virtual void OnDrag( UViewportInteractor* Interactor ) override;
+	virtual void OnDragRelease( UViewportInteractor* Interactor ) override;
+	virtual class UViewportDragOperationComponent* GetDragOperationComponent() override;
+
+protected:
 
 	// AVREditorFloatingUI overrides
 	virtual void SetupWidgetComponent() override;
 
+private:
+
 	/** Set the color on the dynamic materials of the selection bar */
 	void SetSelectionBarColor( const FLinearColor& LinearColor );
 
@@ -81,6 +91,8 @@ private:
 	UPROPERTY()
 	class UMaterialInstanceDynamic* CloseButtonTranslucentMID;
 
+	UPROPERTY()
+	class UViewportDragOperationComponent* DragOperationComponent;
 
 	/** True if at least one hand's laser is aiming toward the UI */
 	bool bIsLaserAimingTowardUI;
@@ -88,17 +100,32 @@ private:
 	/** Scalar that ramps up toward 1.0 after the user aims toward the UI */
 	float AimingAtMeFadeAlpha;
 
-
 	/** True if we're hovering over the selection bar */
 	bool bIsHoveringOverSelectionBar;
 
 	/** Scalar that will advance toward 1.0 over time as we hover over the selection bar */
 	float SelectionBarHoverAlpha;
 
-
 	/** True if we're hovering over the close button */
 	bool bIsHoveringOverCloseButton;
 
 	/** Scalar that will advance toward 1.0 over time as we hover over the close button */
 	float CloseButtonHoverAlpha;
+
+	/** Distance from interactor laser to the handle when starting dragging */
+	float DockSelectDistance;
+};
+
+/**
+ * Calculation for dragging a dockable window
+ */
+UCLASS()
+class UDockableWindowDragOperation : public UViewportDragOperation
+{
+	GENERATED_BODY()
+
+public:
+
+	// IViewportDragOperation
+	virtual void ExecuteDrag( class UViewportInteractor* Interactor, IViewportInteractableInterface* Interactable ) override;
 };
\ No newline at end of file
diff --git a/Engine/Source/Editor/VREditor/VREditorFloatingUI.cpp b/Engine/Source/Editor/VREditor/VREditorFloatingUI.cpp
--- a/Engine/Source/Editor/VREditor/VREditorFloatingUI.cpp
+++ b/Engine/Source/Editor/VREditor/VREditorFloatingUI.cpp
@@ -7,6 +7,7 @@
 #include "VREditorMode.h"
 #include "WidgetComponent.h"
 #include "VREditorWidgetComponent.h"
+#include "VREditorInteractor.h"
 #include "SLevelViewport.h"	// For tab manager tricks
 
 
@@ -19,28 +20,18 @@ namespace VREd
 
 
 AVREditorFloatingUI::AVREditorFloatingUI()
-	: LocalRotation( FRotator( 90.0f, 180.0f, 0.0f ) ),
-	  RelativeOffset( FVector::ZeroVector ),
-	  SlateWidget( nullptr ),
-	  UserWidget( nullptr ),
-	  WidgetComponent( nullptr ),
-	  Scale( 1.0f ),
-	  Resolution( 0, 0 ),
-	  Owner( nullptr ),
-	  DockedTo( EDockedTo::Nothing ),
-	  PreviousDockedTo( EDockedTo::Nothing ),
-	  UserWidgetClass( nullptr ),
-	  bShouldBeVisible(),
-	  FadeAlpha( 1.0f ),
-	  bCollisionOnShowUI( true ),
-	  FadeDelay( 0.0f ),
-	  InitialScale( 1.0f ),
-	  bIsMoving( false ),
-	  MoveToTransform( FTransform() ),
-	  StartMoveToTransform( FTransform() ),
-	  MoveToAlpha( 0.0f ),
-	  MoveToTime( 0.0f ),
-	  MoveToResultDock( EDockedTo::Nothing )
+	: Super(),
+	SlateWidget( nullptr ),
+	UserWidget( nullptr ),
+	WidgetComponent( nullptr ),
+	Resolution( 0, 0 ),
+	Owner( nullptr ),
+	UserWidgetClass( nullptr ),
+	bShouldBeVisible(),
+	FadeAlpha( 1.0f ),
+	bCollisionOnShowUI( true ),
+	FadeDelay( 0.0f ),
+	InitialScale( 1.0f )
 {
 	const bool bTransient = true;
 	USceneComponent* SceneComponent = CreateDefaultSubobject<USceneComponent>( TEXT( "SceneComponent" ), bTransient );
@@ -108,9 +99,10 @@ void AVREditorFloatingUI::SetupWidgetComponent()
 	UpdateTransformIfDocked();
 }
 
-void AVREditorFloatingUI::SetSlateWidget( FVREditorUISystem& InitOwner, const TSharedRef<SWidget>& InitSlateWidget, const FIntPoint InitResolution, const float InitScale, const EDockedTo InitDockedTo )
+void AVREditorFloatingUI::SetSlateWidget( UVREditorUISystem& InitOwner, const TSharedRef<SWidget>& InitSlateWidget, const FIntPoint InitResolution, const float InitScale, const EDockedTo InitDockedTo )
 {
 	Owner = &InitOwner;
+	SetVRMode( &Owner->GetOwner() );
 
 	SlateWidget = InitSlateWidget;
 
@@ -120,15 +112,16 @@ void AVREditorFloatingUI::SetSlateWidget( FVREditorUISystem& InitOwner, const TS
 	Scale = InitScale;
 	InitialScale = Scale;
 
-	DockedTo = InitDockedTo;
+	SetDockedTo( InitDockedTo );
 
 	SetupWidgetComponent();
 }
 
 
-void AVREditorFloatingUI::SetUMGWidget( FVREditorUISystem& InitOwner, TSubclassOf<UVREditorBaseUserWidget> InitUserWidgetClass, const FIntPoint InitResolution, const float InitScale, const EDockedTo InitDockedTo )
+void AVREditorFloatingUI::SetUMGWidget( UVREditorUISystem& InitOwner, TSubclassOf<UVREditorBaseUserWidget> InitUserWidgetClass, const FIntPoint InitResolution, const float InitScale, const EDockedTo InitDockedTo )
 {
 	Owner = &InitOwner;
+	SetVRMode( &Owner->GetOwner() );
 
 	check( InitUserWidgetClass != nullptr );
 	UserWidgetClass = InitUserWidgetClass;
@@ -139,12 +132,11 @@ void AVREditorFloatingUI::SetUMGWidget( FVREditorUISystem& InitOwner, TSubclassO
 	Scale = InitScale;
 	InitialScale = Scale;
 
-	DockedTo = InitDockedTo;
+	SetDockedTo( InitDockedTo );
 
 	SetupWidgetComponent();
 }
 
-
 void AVREditorFloatingUI::Destroyed()
 {
 	if( WidgetComponent != nullptr )
@@ -168,23 +160,18 @@ void AVREditorFloatingUI::Destroyed()
 	Super::Destroyed();
 }
 
-
-void AVREditorFloatingUI::TickManually( float DeltaTime )
+void AVREditorFloatingUI::SetTransform( const FTransform& Transform )
 {
-	Super::Tick( DeltaTime );
+	const FVector AnimatedScale = CalculateAnimatedScale();
+
+	FTransform AnimatedTransform = Transform;
+	AnimatedTransform.SetScale3D( AnimatedTransform.GetScale3D() * AnimatedScale );
 
-	// Update fading state
-	UpdateFadingState( DeltaTime );
+	const float Aspect = ( float ) Resolution.X / ( float ) Resolution.Y;
 
-	if ( bIsMoving )
-	{
-		TickMoveTo( DeltaTime );
-	}
-	else
-	{
-		// Update transform
-		UpdateTransformIfDocked();
-	}
+	RootComponent->SetWorldLocation( AnimatedTransform.GetLocation() );
+	RootComponent->SetWorldRotation( AnimatedTransform.GetRotation() );
+	WidgetComponent->SetWorldScale3D( FVector( 1.0f / AnimatedTransform.GetScale3D().X, 1.0f / ( float ) Resolution.X, 1.0f / ( float ) Resolution.Y / Aspect ) * AnimatedTransform.GetScale3D() );
 }
 
 
@@ -238,51 +225,6 @@ void AVREditorFloatingUI::UpdateFadingState( const float DeltaTime )
 		WidgetComponent->SetTintColorAndOpacity( FLinearColor( UIBrightness, UIBrightness, UIBrightness ).CopyWithNewOpacity( FadeAlpha ) );
 	}
 }
-	
-
-void AVREditorFloatingUI::UpdateTransformIfDocked()
-{
-	if( DockedTo != EDockedTo::Nothing && DockedTo != EDockedTo::Dragging )
-	{
-		FTransform NewTransform;
-		
-		if( DockedTo == EDockedTo::LeftHand )
-		{
-			const bool bOnArm = false;
-			NewTransform = MakeUITransformLockedToHand( VREditorConstants::LeftHandIndex, bOnArm );
-		}
-		else if( DockedTo == EDockedTo::RightHand )
-		{
-			const bool bOnArm = false;
-			NewTransform = MakeUITransformLockedToHand( VREditorConstants::RightHandIndex, bOnArm );
-		}
-		else if( DockedTo == EDockedTo::LeftArm )
-		{
-			const bool bOnArm = true;
-			NewTransform = MakeUITransformLockedToHand( VREditorConstants::LeftHandIndex, bOnArm );
-		}
-		else if( DockedTo == EDockedTo::RightArm )
-		{
-			const bool bOnArm = true;
-			NewTransform = MakeUITransformLockedToHand( VREditorConstants::RightHandIndex, bOnArm );
-		}
-		else if( DockedTo == EDockedTo::Room)
-		{
-			NewTransform = MakeUITransformLockedToRoom();
-		}
-		else if( DockedTo == EDockedTo::Custom )
-		{
-			NewTransform = MakeCustomUITransform();
-		}
-		else
-		{
-			check( 0 );
-		}
-
-		SetTransform( NewTransform );
-	}
-}
-
 
 FVector AVREditorFloatingUI::CalculateAnimatedScale() const
 {
@@ -300,32 +242,6 @@ FVector AVREditorFloatingUI::CalculateAnimatedScale() const
 	return AnimatedScale;
 }
 
-
-void AVREditorFloatingUI::SetTransform( const FTransform& Transform )
-{
-	const FVector AnimatedScale = CalculateAnimatedScale();
-	
-	FTransform AnimatedTransform = Transform;
-	AnimatedTransform.SetScale3D( AnimatedTransform.GetScale3D() * AnimatedScale );
-
-	const float Aspect = (float)Resolution.X / (float)Resolution.Y;
-
-	RootComponent->SetWorldLocation( AnimatedTransform.GetLocation() );
-	RootComponent->SetWorldRotation( AnimatedTransform.GetRotation() );
-	WidgetComponent->SetWorldScale3D( FVector( 1.0f / AnimatedTransform.GetScale3D().X, 1.0f / (float)Resolution.X, 1.0f / (float)Resolution.Y / Aspect ) * AnimatedTransform.GetScale3D() );
-}
-
-
-void AVREditorFloatingUI::SetRelativeOffset( const FVector& InRelativeOffset )
-{
-	RelativeOffset = InRelativeOffset;
-}
-
-void AVREditorFloatingUI::SetLocalRotation( const FRotator& InLocalRotation )
-{
-	LocalRotation = InLocalRotation;
-}
-
 void AVREditorFloatingUI::SetCollisionOnShow( const bool bInCollisionOnShow )
 {
 	bCollisionOnShowUI = bInCollisionOnShow;
@@ -336,88 +252,6 @@ float AVREditorFloatingUI::GetInitialScale() const
 	return InitialScale;
 }
 
-FTransform AVREditorFloatingUI::MakeUITransformLockedToHand( const int32 HandIndex, const bool bOnArm )
-{
-	return MakeUITransformLockedToHand( HandIndex, bOnArm, RelativeOffset, LocalRotation );
-}
-
-FTransform AVREditorFloatingUI::MakeUITransformLockedToHand( const int32 HandIndex, const bool bOnArm, const FVector& InRelativeOffset, const FRotator& InLocalRotation )
-{
-	const float WorldScaleFactor = Owner->GetOwner().GetWorldScaleFactor();
-
-	FTransform UIToHandTransform( InLocalRotation, InRelativeOffset * WorldScaleFactor );
-	if (!bOnArm)
-	{
-		UIToHandTransform *= FTransform( FRotator( VREd::UIOnHandRotationOffset->GetFloat(), 0.0f, 0.0f ).Quaternion(), FVector::ZeroVector );
-	}
-
-	const FTransform HandToWorldTransform = Owner->GetOwner().GetVirtualHand( HandIndex ).Transform;
-
-	FTransform UIToWorldTransform = UIToHandTransform * HandToWorldTransform;
-	UIToWorldTransform.SetScale3D( FVector( Scale * WorldScaleFactor ) );
-
-	return UIToWorldTransform;
-}
-
-void AVREditorFloatingUI::MoveTo( const FTransform& ResultTransform, const float TotalMoveToTime, const EDockedTo ResultDock )
-{
-	MoveToTime = TotalMoveToTime;
-	bIsMoving = true;
-	StartMoveToTransform = GetActorTransform();
-	MoveToTransform = ResultTransform;
-	MoveToAlpha = 0.0f;
-	MoveToResultDock = ResultDock;
-}
-
-void AVREditorFloatingUI::StopMoveTo()
-{
-	bIsMoving = false;
-	SetTransform( MoveToTransform );
-	MoveToTransform = FTransform();
-}
-
-FTransform AVREditorFloatingUI::MakeUITransformLockedToRoom()
-{
-	const float WorldScaleFactor = Owner->GetOwner().GetWorldScaleFactor();
-
-	const FTransform UIToRoomTransform( LocalRotation, RelativeOffset * WorldScaleFactor );
-
-	const FTransform RoomToWorldTransform = Owner->GetOwner().GetRoomTransform();
-
-	FTransform UIToWorldTransform = UIToRoomTransform * RoomToWorldTransform;
-	UIToWorldTransform.SetScale3D( FVector( Scale * WorldScaleFactor ) );
-
-	return UIToWorldTransform;
-}
-
-void AVREditorFloatingUI::TickMoveTo( const float DeltaTime )
-{
-	const float WorldScaleFactor = Owner->GetOwner().GetWorldScaleFactor();
-
-	MoveToAlpha += DeltaTime;
-	const float LerpTime = MoveToTime;
-	if (MoveToAlpha > MoveToTime)
-	{
-		MoveToAlpha = MoveToTime;
-		bIsMoving = false;
-		SetDockedTo( MoveToResultDock );
-	}
-
-	const float CurrentALpha = MoveToAlpha / LerpTime;
-	const FVector NewLocation = FMath::Lerp( StartMoveToTransform.GetLocation(), MoveToTransform.GetLocation(), CurrentALpha );
-	const FQuat NewRotation = FMath::Lerp( StartMoveToTransform.GetRotation(), MoveToTransform.GetRotation(), CurrentALpha );
-	FTransform NewTransform( NewRotation, NewLocation, FVector( Scale * WorldScaleFactor ) );
-
-	SetTransform( NewTransform );
-}
-
-void AVREditorFloatingUI::SetDockedTo( const EDockedTo NewDockedTo )
-{
-	PreviousDockedTo = DockedTo;
-	DockedTo = NewDockedTo;
-	UpdateTransformIfDocked();
-}
-
 
 void AVREditorFloatingUI::ShowUI( const bool bShow, const bool bAllowFading, const float InitFadeDelay )
 {
diff --git a/Engine/Source/Editor/VREditor/VREditorFloatingUI.h b/Engine/Source/Editor/VREditor/VREditorFloatingUI.h
--- a/Engine/Source/Editor/VREditor/VREditorFloatingUI.h
+++ b/Engine/Source/Editor/VREditor/VREditorFloatingUI.h
@@ -2,7 +2,7 @@
 
 #pragma once
 
-#include "GameFramework/Actor.h"
+#include "VREditorBaseActor.h"
 #include "VREditorFloatingUI.generated.h"
 
 
@@ -10,7 +10,7 @@
  * Represents an interactive floating UI panel in the VR Editor
  */
 UCLASS()
-class AVREditorFloatingUI : public AActor
+class AVREditorFloatingUI : public AVREditorBaseActor
 {
 	GENERATED_BODY()
 
@@ -19,24 +19,11 @@ public:
 	/** Default constructor which sets up safe defaults */
 	AVREditorFloatingUI();
 
-	/** Possible UI attachment points */
-	enum class EDockedTo
-	{
-		Nothing,
-		LeftHand,
-		RightHand,
-		LeftArm,
-		RightArm,
-		Room,
-		Custom,
-		Dragging,
-	};
-
 	/** Creates a FVREditorFloatingUI using a Slate widget, and sets up safe defaults */
-	void SetSlateWidget( class FVREditorUISystem& InitOwner, const TSharedRef<SWidget>& InitSlateWidget, const FIntPoint InitResolution, const float InitScale, const EDockedTo InitDockedTo );
+	void SetSlateWidget( class UVREditorUISystem& InitOwner, const TSharedRef<SWidget>& InitSlateWidget, const FIntPoint InitResolution, const float InitScale, const EDockedTo InitDockedTo );
 
 	/** Creates a FVREditorFloatingUI using a UMG user widget, and sets up safe defaults */
-	void SetUMGWidget( class FVREditorUISystem& InitOwner, class TSubclassOf<class UVREditorBaseUserWidget> InitUserWidgetClass, const FIntPoint InitResolution, const float InitScale, const EDockedTo InitDockedTo );
+	void SetUMGWidget( class UVREditorUISystem& InitOwner, class TSubclassOf<class UVREditorBaseUserWidget> InitUserWidgetClass, const FIntPoint InitResolution, const float InitScale, const EDockedTo InitDockedTo );
 
 	/** @return Returns true if the UI is visible (or wants to be visible -- it might be transitioning */
 	bool IsUIVisible() const
@@ -47,21 +34,6 @@ public:
 	/** Shows or hides the UI (also enables collision, and performs a transition effect) */
 	void ShowUI( const bool bShow, const bool bAllowFading = true, const float InitFadeDelay = 0.0f );
 
-	/** Returns what we're docked to, if anything */
-	EDockedTo GetDockedTo() const
-	{
-		return DockedTo;
-	}
-
-	/** Returns what we were previously docked to */
-	EDockedTo GetPreviouslyDockedTo() const
-	{
-		return PreviousDockedTo;
-	}
-
-	/** Sets what this UI is docked to */
-	void SetDockedTo( const EDockedTo NewDockedTo );
-
 	/** Returns the widget component for this UI, or nullptr if not spawned right now */
 	class UVREditorWidgetComponent* GetWidgetComponent()
 	{
@@ -77,25 +49,21 @@ public:
 	/** Sets a new size for the UI */
 	void SetScale( const float NewSize );
 
+	/** Sets the UI transform */
+	virtual void SetTransform( const FTransform& Transform ) override;
+
 	/** AActor overrides */
 	virtual void Destroyed() override;
 
-	/** Called by the UI system to tick this UI every frame.  We can't use AActor::Tick() because this actor can exist in 
-	    the editor's world, which never ticks. */
-	virtual void TickManually( float DeltaTime );
-
-	/** Sets the transform for this UI */
-	void SetTransform( const FTransform& Transform );
-
 	/** Returns the owner of this object */
-	FVREditorUISystem& GetOwner()
+	UVREditorUISystem& GetOwner()
 	{
 		check( Owner != nullptr );
 		return *Owner;
 	}
 
 	/** Returns the owner of this object (const) */
-	const FVREditorUISystem& GetOwner() const
+	const UVREditorUISystem& GetOwner() const
 	{
 		check( Owner != nullptr );
 		return *Owner;
@@ -108,62 +76,24 @@ public:
 		return CastChecked<T>( UserWidget );
 	}
 
-	/** Sets the relative offset of the UI */
-	void SetRelativeOffset( const FVector& InRelativeOffset );
-
-	/** Sets the local rotation of the UI */
-	void SetLocalRotation( const FRotator& InRelativeRotation );
-
 	/** Sets if the collision is on when showing the UI*/
 	void SetCollisionOnShow( const bool bInCollisionOnShow );
 
 	/** Gets the initial size of this UI */
 	float GetInitialScale() const;
 
-	/** Get a relative transform from the hand */
-	FTransform MakeUITransformLockedToHand( const int32 HandIndex, const bool bOnArm, const FVector& InRelativeOffset, const FRotator& InLocalRotation );
-
-	/** Start moving towards a transform */
-	void MoveTo( const FTransform& ResultTransform, const float TotalMoveToTime, const EDockedTo ResultDock );
-
-	/** Abort moving to transform and move instant to MoveToTransform */
-	void StopMoveTo();
-
 protected:
 
 	/** Returns a scale to use for this widget that takes into account animation */
 	FVector CalculateAnimatedScale() const;
 
-	/** Called when DockTo state is custom, needs to be implemented by derived classes */
-	virtual FTransform MakeCustomUITransform() { return FTransform(); };
-
 	/** Called to finish setting everything up, after a widget has been assigned */
 	virtual void SetupWidgetComponent();
 
 private:
 
 	/** Called after spawning, and every Tick, to update opacity of the widget */
-	void UpdateFadingState( const float DeltaTime );
-
-	/** Called every tick to keep the UI position up to date */
-	void UpdateTransformIfDocked();
-
-	/** Given a hand to lock to, returns a transform to place UI at that hand's location and orientation */
-	FTransform MakeUITransformLockedToHand( const int32 HandIndex, const bool bOnArm);
-
-	/** Create a room transform with the RelativeOffset and LocalRotation */
-	FTransform MakeUITransformLockedToRoom();
-
-	/** Updates the lerp movement */
-	void TickMoveTo( const float DeltaTime );
-
-public:
-
-	/** Local rotation of the UI */
-	FRotator LocalRotation;
-
-	/** Relative offset of the UI */
-	FVector RelativeOffset;
+	virtual void UpdateFadingState( const float DeltaTime );
 
 protected:
 	/** Slate widget we're drawing, or null if we're drawing a UMG user widget */
@@ -176,9 +106,6 @@ protected:
 	/** When in a spawned state, this is the widget component to represent the widget */
 	UPROPERTY()
 	class UVREditorWidgetComponent* WidgetComponent;
-	
-	/** How big the UI should be */
-	float Scale;
 
 	/** Resolution we should draw this UI at, regardless of scale */
 	FIntPoint Resolution;
@@ -186,13 +113,7 @@ protected:
 private:
 
 	/** Owning object */
-	class FVREditorUISystem* Owner;
-
-	/** What the UI is attached to */
-	EDockedTo DockedTo;
-
-	/** What was the UI previously attached to */
-	EDockedTo PreviousDockedTo;
+	class UVREditorUISystem* Owner;
 
 	/** Class of the UMG widget we're spawning */
 	UPROPERTY()
@@ -213,23 +134,5 @@ private:
 
 	/** The starting scale of this UI */
 	float InitialScale;
-
-	/** If the UI is currently moving */
-	bool bIsMoving;
-
-	/** The end transform to move to */
-	FTransform MoveToTransform;
-
-	/** The actor transform when started moving */
-	FTransform StartMoveToTransform;
-
-	/** Current lerp alpha */
-	float MoveToAlpha;
-
-	/** Total time to move to the end transform */
-	float MoveToTime;
-
-	/** Dock state to switch to when finished moving */
-	EDockedTo MoveToResultDock;
 };
 
