Small clean up for MeshPaintEd

diff --git a/Engine/Source/Editor/MeshPaint/Private/MeshPaintEdMode.cpp b/Engine/Source/Editor/MeshPaint/Private/MeshPaintEdMode.cpp
--- a/Engine/Source/Editor/MeshPaint/Private/MeshPaintEdMode.cpp
+++ b/Engine/Source/Editor/MeshPaint/Private/MeshPaintEdMode.cpp
@@ -193,7 +193,7 @@ void FEdModeMeshPaint::Enter()
 		SetViewportShowFlags( bAllowColorViewModes, *ViewportClient );
 	} 
 
-	//When painting vertext colors we want to force the lod level of objects being painted to LOD0.
+	//When painting vertex colors we want to force the lod level of objects being painted to LOD0.
 	if( FMeshPaintSettings::Get().ResourceType == EMeshPaintResource::VertexColors )
 	{
 		ApplyOrRemoveForceBestLOD(/*bApply=*/ true);
@@ -2699,7 +2699,7 @@ void FEdModeMeshPaint::Render( const FSceneView* View, FViewport* Viewport, FPri
 						float StrengthScale = ( PaintingWithInteractorInVR == nullptr && FMeshPaintSettings::Get().bEnableFlow ) ? FMeshPaintSettings::Get().FlowAmount : 1.0f;
 
 						// Apply VR controller trigger pressure if we're painting in VR
-						if ( bIsPainting && PaintingWithInteractorInVR && VREditorMode != nullptr && PaintingWithInteractorInVR )
+						if ( bIsPainting && PaintingWithInteractorInVR && VREditorMode != nullptr )
 						{
 							UVREditorInteractor* VRInteractor = Cast<UVREditorInteractor>( PaintingWithInteractorInVR );
 							if ( VRInteractor )
