VREditor: Fixed mesh paint brush not showing brush preview while in normal editor.

diff --git a/Engine/Source/Editor/MeshPaint/Private/MeshPaintEdMode.cpp b/Engine/Source/Editor/MeshPaint/Private/MeshPaintEdMode.cpp
--- a/Engine/Source/Editor/MeshPaint/Private/MeshPaintEdMode.cpp
+++ b/Engine/Source/Editor/MeshPaint/Private/MeshPaintEdMode.cpp
@@ -2634,16 +2634,16 @@ void FEdModeMeshPaint::Render( const FSceneView* View, FViewport* Viewport, FPri
 		const bool bAllowColorViewModes = ( FMeshPaintSettings::Get().ResourceType != EMeshPaintResource::Texture );
 		SetViewportShowFlags( bAllowColorViewModes, *ViewportClient );
 
-
 		bool bHavePaintRay = false;
 		FVector PaintRayViewOrigin = FVector::ZeroVector;
 		FVector PaintRayStart = FVector::ZeroVector;
 		FVector PaintRayDirection = FVector::ZeroVector;
 
-		// Check to see if VREditorMode is active
+		IVREditorMode* VREditorMode = static_cast<IVREditorMode*>(GetModeManager()->GetActiveMode(IVREditorModule::Get().GetVREditorModeID()));
+
+		// Check to see if VREditorMode is active. If so, we're painting with interactor
 		if ( IVREditorModule::IsAvailable() )
 		{
-			IVREditorMode* VREditorMode = static_cast< IVREditorMode* >( GetModeManager()->GetActiveMode( IVREditorModule::Get().GetVREditorModeID() ) );
 			if ( VREditorMode != nullptr && VREditorMode->IsFullyInitialized() )
 			{
 				for ( UViewportInteractor* Interactor : VREditorMode->GetWorldInteraction().GetInteractors() )
@@ -2664,67 +2664,69 @@ void FEdModeMeshPaint::Render( const FSceneView* View, FViewport* Viewport, FPri
 						PaintRayStart = LaserPointerStart;
 						PaintRayDirection = LaserPointerDirection;
 					}
+				}
+			}
+		}
 
-					if ( !bHavePaintRay )
-					{
-						// Make sure the cursor is visible OR we're flood filling.  No point drawing a paint cue when there's no cursor.
-						if ( Viewport->IsCursorVisible() || IsForceRendered() )
-						{
-							if ( !PDI->IsHitTesting() )
-							{
-								// Grab the mouse cursor position
-								FIntPoint MousePosition;
-								Viewport->GetMousePos( MousePosition );
-
-								// Is the mouse currently over the viewport? or flood filling
-								if ( IsForceRendered() || ( MousePosition.X >= 0 && MousePosition.Y >= 0 && MousePosition.X < ( int32 ) Viewport->GetSizeXY().X && MousePosition.Y < ( int32 ) Viewport->GetSizeXY().Y ) )
-								{
-									// Compute a world space ray from the screen space mouse coordinates
-									FViewportCursorLocation MouseViewportRay( View, ViewportClient, MousePosition.X, MousePosition.Y );
-
-									bHavePaintRay = true;
-									PaintRayViewOrigin = View->ViewMatrices.ViewOrigin;
-									PaintRayStart = MouseViewportRay.GetOrigin();
-									PaintRayDirection = MouseViewportRay.GetDirection();
-								}
-							}
-						}
-					}
+		// Interactor not available, check to see if we're painting with mouse
+		if (!bHavePaintRay)
+		{
+			// Make sure the cursor is visible OR we're flood filling.  No point drawing a paint cue when there's no cursor.
+			if (Viewport->IsCursorVisible() || IsForceRendered())
+			{
+				if (!PDI->IsHitTesting())
+				{
+					// Grab the mouse cursor position
+					FIntPoint MousePosition;
+					Viewport->GetMousePos(MousePosition);
 
-					if ( bHavePaintRay )
+					// Is the mouse currently over the viewport? or flood filling
+					if (IsForceRendered() || (MousePosition.X >= 0 && MousePosition.Y >= 0 && MousePosition.X < (int32)Viewport->GetSizeXY().X && MousePosition.Y < (int32)Viewport->GetSizeXY().Y))
 					{
-						// Unless "Flow" mode is enabled, we'll only draw a visual cue while rendering and won't
-						// do any actual painting.  When "Flow" is turned on we will paint here, too!
-						const bool bVisualCueOnly = !bIsPainting || ( PaintingWithInteractorInVR == nullptr && !FMeshPaintSettings::Get().bEnableFlow );
-						float StrengthScale = ( PaintingWithInteractorInVR == nullptr && FMeshPaintSettings::Get().bEnableFlow ) ? FMeshPaintSettings::Get().FlowAmount : 1.0f;
+						// Compute a world space ray from the screen space mouse coordinates
+						FViewportCursorLocation MouseViewportRay(View, ViewportClient, MousePosition.X, MousePosition.Y);
 
-						// Apply VR controller trigger pressure if we're painting in VR
-						if ( bIsPainting && PaintingWithInteractorInVR && VREditorMode != nullptr )
-						{
-							UVREditorInteractor* VRInteractor = Cast<UVREditorInteractor>( PaintingWithInteractorInVR );
-							if ( VRInteractor )
-							{
-								StrengthScale *= VRInteractor->GetSelectAndMoveTriggerValue();
+						bHavePaintRay = true;
+						PaintRayViewOrigin = View->ViewMatrices.ViewOrigin;
+						PaintRayStart = MouseViewportRay.GetOrigin();
+						PaintRayDirection = MouseViewportRay.GetDirection();
+					}
+				}
+			}
+		}
 
-								// Make sure light press locking is enabled in VR for the hand we're painting with.  We don't want
-								// a full press to be converted to a move action, however locking will have been disabled already
-								// by the world interaction code after a light press happened on an object. 
-								VRInteractor->SetAllowTriggerLightPressLocking( true );
-							}
-						}
+		// Apply paint pressure and start painting
+		if (bHavePaintRay)
+		{
+			// Unless "Flow" mode is enabled, we'll only draw a visual cue while rendering and won't
+			// do any actual painting.  When "Flow" is turned on we will paint here, too!
+			const bool bVisualCueOnly = !bIsPainting || (PaintingWithInteractorInVR == nullptr && !FMeshPaintSettings::Get().bEnableFlow);
+			float StrengthScale = (PaintingWithInteractorInVR == nullptr && FMeshPaintSettings::Get().bEnableFlow) ? FMeshPaintSettings::Get().FlowAmount : 1.0f;
 
-						// Apply stylus pressure if it's active
-						else if ( Viewport->IsPenActive() )
-						{
-							StrengthScale *= Viewport->GetTabletPressure();
-						}
+			// Apply VR controller trigger pressure if we're painting in VR
+			if (bIsPainting && PaintingWithInteractorInVR && VREditorMode != nullptr)
+			{
+				UVREditorInteractor* VRInteractor = Cast<UVREditorInteractor>(PaintingWithInteractorInVR);
+				if (VRInteractor)
+				{
+					StrengthScale *= VRInteractor->GetSelectAndMoveTriggerValue();
 
-						const EMeshPaintAction::Type PaintAction = GetPaintAction( Viewport );
-						bool bAnyPaintableActorsUnderCursor = false;
-						DoPaint( PaintRayViewOrigin, PaintRayStart, PaintRayDirection, PDI, PaintAction, bVisualCueOnly, StrengthScale, /* Out */ bAnyPaintableActorsUnderCursor );
-					}
+					// Make sure light press locking is enabled in VR for the hand we're painting with.  We don't want
+					// a full press to be converted to a move action, however locking will have been disabled already
+					// by the world interaction code after a light press happened on an object. 
+					VRInteractor->SetAllowTriggerLightPressLocking(true);
 				}
 			}
+
+			// Apply stylus pressure if it's active
+			else if (Viewport->IsPenActive())
+			{
+				StrengthScale *= Viewport->GetTabletPressure();
+			}
+
+			const EMeshPaintAction::Type PaintAction = GetPaintAction(Viewport);
+			bool bAnyPaintableActorsUnderCursor = false;
+			DoPaint(PaintRayViewOrigin, PaintRayStart, PaintRayDirection, PDI, PaintAction, bVisualCueOnly, StrengthScale, /* Out */ bAnyPaintableActorsUnderCursor);
 		}
 	}
 }
