Mesh Paint: Fixed various bugs
- Fixed brush preview not rendered for lasers/mouse when not actively painting (UE-33554)
- Fixed 'full press' over actors preventing UI from being clickable (UE-33550)
- Fixed brush cursor displayed when hovering over UI (including selection bar/close button) (UE-33551)
- Fixed VR transform gizmo getting in the way of everything while painting (it is now hidden while in mesh paint mode)
- Fixed not being able to interact with UIs after messing around with mesh paint (UE-33621)

diff --git a/Engine/Source/Editor/MeshPaint/Private/MeshPaintEdMode.cpp b/Engine/Source/Editor/MeshPaint/Private/MeshPaintEdMode.cpp
--- a/Engine/Source/Editor/MeshPaint/Private/MeshPaintEdMode.cpp
+++ b/Engine/Source/Editor/MeshPaint/Private/MeshPaintEdMode.cpp
@@ -33,7 +33,9 @@
 #include "Components/SplineMeshComponent.h"
 
 #include "ViewportWorldInteraction.h"
+#include "ViewportInteractableInterface.h"
 #include "VREditorInteractor.h"
+#include "VIBaseTransformGizmo.h"
 
 #define LOCTEXT_NAMESPACE "MeshPaint_Mode"
 
@@ -208,18 +210,24 @@ void FEdModeMeshPaint::Enter()
 	{
 		VREditorMode->GetWorldInteraction().OnViewportInteractionInputAction().RemoveAll( this );
 		VREditorMode->GetWorldInteraction().OnViewportInteractionInputAction().AddRaw( this, &FEdModeMeshPaint::OnVRAction );
+
+		// Hide the VR transform gizmo while we're in mesh paint mode.  It sort of gets in the way of painting.
+		VREditorMode->GetWorldInteraction().SetTransformGizmoVisible( false );
 	}
 }
 
 /** FEdMode: Called when the mode is exited */
 void FEdModeMeshPaint::Exit()
 {
-	// Unregister from event handlers
 	if( IVREditorModule::IsAvailable() )
 	{
 		IVREditorMode* VREditorMode = static_cast<IVREditorMode*>( GetModeManager()->GetActiveMode( IVREditorModule::Get().GetVREditorModeID() ) );
 		if( VREditorMode != nullptr )
 		{
+			// Restore the transform gizmo visibility
+			VREditorMode->GetWorldInteraction().SetTransformGizmoVisible( true );
+
+			// Unregister from event handlers
 			VREditorMode->GetWorldInteraction().OnViewportInteractionInputAction().RemoveAll( this );
 		}
 	}
@@ -2634,18 +2642,27 @@ void FEdModeMeshPaint::Render( const FSceneView* View, FViewport* Viewport, FPri
 		const bool bAllowColorViewModes = ( FMeshPaintSettings::Get().ResourceType != EMeshPaintResource::Texture );
 		SetViewportShowFlags( bAllowColorViewModes, *ViewportClient );
 
-		bool bHavePaintRay = false;
-		FVector PaintRayViewOrigin = FVector::ZeroVector;
-		FVector PaintRayStart = FVector::ZeroVector;
-		FVector PaintRayDirection = FVector::ZeroVector;
+		struct FPaintRay
+		{
+			FVector CameraLocation;
+			FVector RayStart;
+			FVector RayDirection;
+			UViewportInteractor* ViewportInteractor;
+		};
+
+		static TArray<FPaintRay> PaintRays;
+		PaintRays.Reset();
+		
 
 		IVREditorMode* VREditorMode = static_cast<IVREditorMode*>(GetModeManager()->GetActiveMode(IVREditorModule::Get().GetVREditorModeID()));
 
 		// Check to see if VREditorMode is active. If so, we're painting with interactor
+		bool bIsInVRMode = false;
 		if ( IVREditorModule::IsAvailable() )
 		{
 			if ( VREditorMode != nullptr && VREditorMode->IsFullyInitialized() )
 			{
+				bIsInVRMode = true;
 				for ( UViewportInteractor* Interactor : VREditorMode->GetWorldInteraction().GetInteractors() )
 				{
 					// Skip other interactors if we are painting with one
@@ -2659,17 +2676,18 @@ void FEdModeMeshPaint::Render( const FSceneView* View, FViewport* Viewport, FPri
 					{
 						const FVector LaserPointerDirection = ( LaserPointerEnd - LaserPointerStart ).GetSafeNormal();
 
-						bHavePaintRay = true;
-						PaintRayViewOrigin = VREditorMode->GetHeadTransform().GetLocation();
-						PaintRayStart = LaserPointerStart;
-						PaintRayDirection = LaserPointerDirection;
+						FPaintRay& NewPaintRay = *new( PaintRays ) FPaintRay();
+						NewPaintRay.CameraLocation = VREditorMode->GetHeadTransform().GetLocation();
+						NewPaintRay.RayStart = LaserPointerStart;
+						NewPaintRay.RayDirection = LaserPointerDirection;
+						NewPaintRay.ViewportInteractor = Interactor;
 					}
 				}
 			}
 		}
 
-		// Interactor not available, check to see if we're painting with mouse
-		if (!bHavePaintRay)
+		// Check to see if we're painting with mouse
+		if (!bIsInVRMode)
 		{
 			// Make sure the cursor is visible OR we're flood filling.  No point drawing a paint cue when there's no cursor.
 			if (Viewport->IsCursorVisible() || IsForceRendered())
@@ -2686,47 +2704,73 @@ void FEdModeMeshPaint::Render( const FSceneView* View, FViewport* Viewport, FPri
 						// Compute a world space ray from the screen space mouse coordinates
 						FViewportCursorLocation MouseViewportRay(View, ViewportClient, MousePosition.X, MousePosition.Y);
 
-						bHavePaintRay = true;
-						PaintRayViewOrigin = View->ViewMatrices.ViewOrigin;
-						PaintRayStart = MouseViewportRay.GetOrigin();
-						PaintRayDirection = MouseViewportRay.GetDirection();
+						FPaintRay& NewPaintRay = *new( PaintRays ) FPaintRay();
+						NewPaintRay.CameraLocation = View->ViewMatrices.ViewOrigin;
+						NewPaintRay.RayStart = MouseViewportRay.GetOrigin();
+						NewPaintRay.RayDirection = MouseViewportRay.GetDirection();
+						NewPaintRay.ViewportInteractor = nullptr;
 					}
 				}
 			}
 		}
 
-		// Apply paint pressure and start painting
-		if (bHavePaintRay)
+		// Apply paint pressure and start painting (or if not currently painting, draw a preview of where paint will be applied)
+		for( const FPaintRay& PaintRay : PaintRays )
 		{
 			// Unless "Flow" mode is enabled, we'll only draw a visual cue while rendering and won't
 			// do any actual painting.  When "Flow" is turned on we will paint here, too!
 			const bool bVisualCueOnly = !bIsPainting || (PaintingWithInteractorInVR == nullptr && !FMeshPaintSettings::Get().bEnableFlow);
 			float StrengthScale = (PaintingWithInteractorInVR == nullptr && FMeshPaintSettings::Get().bEnableFlow) ? FMeshPaintSettings::Get().FlowAmount : 1.0f;
 
-			// Apply VR controller trigger pressure if we're painting in VR
-			if (bIsPainting && PaintingWithInteractorInVR && VREditorMode != nullptr)
+			bool bIsHoveringOverUIVR = false;
+			const UVREditorInteractor* VRInteractor = Cast<UVREditorInteractor>( PaintRay.ViewportInteractor );
+			if( VRInteractor != nullptr )
+			{
+				bIsHoveringOverUIVR = VRInteractor->IsHoveringOverUI();
+			}
+
+			// Don't draw visual cue for paint brush when the interactor is hovering over UI
+			if( !bVisualCueOnly || PaintRay.ViewportInteractor == nullptr || !bIsHoveringOverUIVR )
 			{
-				UVREditorInteractor* VRInteractor = Cast<UVREditorInteractor>(PaintingWithInteractorInVR);
-				if (VRInteractor)
+				// Don't draw visual cue if we're hovering over a viewport interactable, such as a dockable window selection bar
+				bool bIsHoveringOverViewportInteractable = false;
+				if( bVisualCueOnly && PaintRay.ViewportInteractor != nullptr )
 				{
-					StrengthScale *= VRInteractor->GetSelectAndMoveTriggerValue();
+					FHitResult HitResult = PaintRay.ViewportInteractor->GetHitResultFromLaserPointer();
+					if( HitResult.Actor.IsValid() )
+					{
+						UViewportWorldInteraction& WorldInteraction = VREditorMode->GetWorldInteraction();
+
+						if( WorldInteraction.IsInteractableComponent( HitResult.GetComponent() ) )
+						{
+							AActor* Actor = HitResult.Actor.Get();
 
-					// Make sure light press locking is enabled in VR for the hand we're painting with.  We don't want
-					// a full press to be converted to a move action, however locking will have been disabled already
-					// by the world interaction code after a light press happened on an object. 
-					VRInteractor->SetAllowTriggerLightPressLocking(true);
+							// Make sure we're not hovering over some other viewport interactable, such as a dockable window selection bar or close button
+							IViewportInteractableInterface* ActorInteractable = Cast<IViewportInteractableInterface>( Actor );
+							bIsHoveringOverViewportInteractable = ( ActorInteractable != nullptr );
+						}
+					}
 				}
-			}
 
-			// Apply stylus pressure if it's active
-			else if (Viewport->IsPenActive())
-			{
-				StrengthScale *= Viewport->GetTabletPressure();
-			}
+				if( !bIsHoveringOverViewportInteractable )
+				{
+					// Apply VR controller trigger pressure if we're painting in VR
+					if (bIsPainting && PaintingWithInteractorInVR && VREditorMode != nullptr)
+					{
+						StrengthScale *= VRInteractor->GetSelectAndMoveTriggerValue();
+					}
 
-			const EMeshPaintAction::Type PaintAction = GetPaintAction(Viewport);
-			bool bAnyPaintableActorsUnderCursor = false;
-			DoPaint(PaintRayViewOrigin, PaintRayStart, PaintRayDirection, PDI, PaintAction, bVisualCueOnly, StrengthScale, /* Out */ bAnyPaintableActorsUnderCursor);
+					// Apply stylus pressure if it's active
+					else if (Viewport->IsPenActive())
+					{
+						StrengthScale *= Viewport->GetTabletPressure();
+					}
+
+					const EMeshPaintAction::Type PaintAction = GetPaintAction(Viewport);
+					bool bAnyPaintableActorsUnderCursor = false;
+					DoPaint(PaintRay.CameraLocation, PaintRay.RayStart, PaintRay.RayDirection, PDI, PaintAction, bVisualCueOnly, StrengthScale, /* Out */ bAnyPaintableActorsUnderCursor);
+				}
+			}
 		}
 	}
 }
@@ -5086,44 +5130,66 @@ void FEdModeMeshPaint::OnVRAction( class FEditorViewportClient& ViewportClient,
 
 	if( VREditorMode != nullptr && VRInteractor != nullptr )
 	{
-		// Absorb all "full press" events unless the user is clicking on UI to avoid objects getting moved around.
-		// @todo vreditor: We shouldn't have to use this trick.  This input system needs more flexibility in its public interface!
-		if( Action.ActionType == ViewportWorldActionTypes::SelectAndMove && bIsPainting && 
-			PaintingWithInteractorInVR && PaintingWithInteractorInVR == Interactor && 
-			!VRInteractor->IsHoveringOverUI() )
-		{
-			bWasHandled = true;
-			bOutIsInputCaptured = true;
-		}
-
 		if( Action.ActionType == ViewportWorldActionTypes::SelectAndMove_LightlyPressed )
 		{
 			if( !bIsPainting && Action.Event == IE_Pressed && !VRInteractor->IsHoveringOverUI() )
 			{
-				// NOTE: We intentionally don't capture input for "LightlyPressed" because we want to allow objects to still
-				// become selected when clicked on.  We avoid allowing the objects to be moved by calling the 
-				// SetAllowTriggerLightPressLocking() function in the next tick (after VR Editor will have disabled locking.)
-				// @todo vreditor: We shouldn't have to use this trick.  This input system needs more flexibility in its public interface!
-				//bWasHandled = true;
-				//bOutIsInputCaptured = true;
-
-				StartPainting();
-				PaintingWithInteractorInVR = Interactor;
-
-				// Go ahead and paint immediately
-				FVector LaserPointerStart, LaserPointerEnd;
-				if( Interactor->GetLaserPointer( /* Out */ LaserPointerStart, /* Out */ LaserPointerEnd ) )
+				// Check to see that we're clicking on a selected object.  You can only paint on selected things.  Otherwise,
+				// we'll fall through to the normal interaction code which might cause the object to become selected.
+				bool bIsClickingOnSelectedObject = false;
 				{
-					const FVector LaserPointerDirection = ( LaserPointerEnd - LaserPointerStart ).GetSafeNormal();
+					FHitResult HitResult = VRInteractor->GetHitResultFromLaserPointer();
+					if( HitResult.Actor.IsValid() )
+					{
+						UViewportWorldInteraction& WorldInteraction = VREditorMode->GetWorldInteraction();
+						
+						if( WorldInteraction.IsInteractableComponent( HitResult.GetComponent() ) )
+						{
+							AActor* Actor = HitResult.Actor.Get();
 
-					// Apply VR controller trigger pressure
-					float StrengthScale = VRInteractor->GetSelectAndMoveTriggerValue();
+							// Make sure we're not hovering over some other viewport interactable, such as a dockable window selection bar or close button
+							IViewportInteractableInterface* ActorInteractable = Cast<IViewportInteractableInterface>( Actor );
+							if( ActorInteractable == nullptr )
+							{
+								if( Actor != WorldInteraction.GetTransformGizmoActor() )  // Don't change any actor selection state if the user clicked on a gizmo
+								{
+									if( Actor->IsSelected() )
+									{
+										bIsClickingOnSelectedObject = true;
+									}
+								}
+							}
+						}
+					}
+				}
 
-					// Paint!
-					const bool bVisualCueOnly = false;
-					const EMeshPaintAction::Type PaintAction = GetPaintAction( ViewportClient.Viewport );
-					bool bAnyPaintableActorsUnderCursor = false;
-					DoPaint( VREditorMode->GetHeadTransform().GetLocation(), LaserPointerStart, LaserPointerDirection, NULL, PaintAction, bVisualCueOnly, StrengthScale, /* Out */ bAnyPaintableActorsUnderCursor );
+				if( bIsClickingOnSelectedObject )
+				{
+					bWasHandled = true;
+					bOutIsInputCaptured = true;
+
+					// Don't allow a "full press" to happen at all with this trigger pull.  We don't need full press for anything, and it
+					// would interrupt our light press.
+					VRInteractor->SetAllowTriggerFullPress( false );
+
+					StartPainting();
+					PaintingWithInteractorInVR = Interactor;
+
+					// Go ahead and paint immediately
+					FVector LaserPointerStart, LaserPointerEnd;
+					if( Interactor->GetLaserPointer( /* Out */ LaserPointerStart, /* Out */ LaserPointerEnd ) )
+					{
+						const FVector LaserPointerDirection = ( LaserPointerEnd - LaserPointerStart ).GetSafeNormal();
+
+						// Apply VR controller trigger pressure
+						float StrengthScale = VRInteractor->GetSelectAndMoveTriggerValue();
+
+						// Paint!
+						const bool bVisualCueOnly = false;
+						const EMeshPaintAction::Type PaintAction = GetPaintAction( ViewportClient.Viewport );
+						bool bAnyPaintableActorsUnderCursor = false;
+						DoPaint( VREditorMode->GetHeadTransform().GetLocation(), LaserPointerStart, LaserPointerDirection, NULL, PaintAction, bVisualCueOnly, StrengthScale, /* Out */ bAnyPaintableActorsUnderCursor );
+					}
 				}
 			}
 
@@ -5132,9 +5198,8 @@ void FEdModeMeshPaint::OnVRAction( class FEditorViewportClient& ViewportClient,
 			{
 				EndPainting();
 
-				// @todo vreditor: We shouldn't have to use this trick.  This input system needs more flexibility in its public interface!
-				//bWasHandled = true;
-				//bOutIsInputCaptured = false;
+				bWasHandled = true;
+				bOutIsInputCaptured = false;
 			}
 		}
 	}
diff --git a/Engine/Source/Editor/VREditor/VREditorMotionControllerInteractor.cpp b/Engine/Source/Editor/VREditor/VREditorMotionControllerInteractor.cpp
--- a/Engine/Source/Editor/VREditor/VREditorMotionControllerInteractor.cpp
+++ b/Engine/Source/Editor/VREditor/VREditorMotionControllerInteractor.cpp
@@ -551,6 +551,7 @@ void UVREditorMotionControllerInteractor::HandleInputKey( FViewportActionKeyInpu
 				if ( !IsClickingOnUI() &&	// @todo vreditor: UI clicks happen with light presses; we never want to convert to a hard press!
 					GetDraggingMode() != EViewportInteractionDraggingMode::Material &&	// @todo vreditor: Material dragging happens with light presses, don't convert to a hard press!
 					GetDraggingMode() != EViewportInteractionDraggingMode::ActorsAtLaserImpact &&	// @todo vreditor: Actor dragging happens with light presses, don't convert to a hard press!
+					AllowTriggerFullPress() &&
 					( !AllowTriggerLightPressLocking() || TimeSinceLightlyPressed < VREd::TriggerLightlyPressedLockTime->GetFloat() ) )
 				{
 					SetTriggerAtLeastLightlyPressed( false );
@@ -635,6 +636,7 @@ void UVREditorMotionControllerInteractor::HandleInputAxis( FViewportActionKeyInp
 			{
 				bIsTriggerAtLeastLightlyPressed = true;
 				SetAllowTriggerLightPressLocking( true );
+				SetAllowTriggerFullPress( true );
 				RealTimeTriggerWasLightlyPressed = FPlatformTime::Seconds();
 				bHasTriggerBeenReleasedSinceLastPress = false;
 
@@ -660,7 +662,8 @@ void UVREditorMotionControllerInteractor::HandleInputAxis( FViewportActionKeyInp
 		// Synthesize "fully pressed" events for the trigger
 		{
 			if ( !bIsTriggerFullyPressed &&	// Don't fire unless we are already pressed
-				Delta >= VREd::TriggerFullyPressedThreshold->GetFloat() )
+				 AllowTriggerFullPress() &&
+				 Delta >= VREd::TriggerFullyPressedThreshold->GetFloat() )
 			{
 				bIsTriggerFullyPressed = true;
 
diff --git a/Engine/Source/Editor/ViewportInteraction/Public/ViewportInteractor.h b/Engine/Source/Editor/ViewportInteraction/Public/ViewportInteractor.h
--- a/Engine/Source/Editor/ViewportInteraction/Public/ViewportInteractor.h
+++ b/Engine/Source/Editor/ViewportInteraction/Public/ViewportInteractor.h
@@ -142,18 +142,31 @@ public:
 	float GetDragHapticFeedbackStrength() const;
 
 	/**
-	* Enables or disable "light press" locking for this interactor.  This is an advanced feature.
-	* It should be true if we allow locking of the lightly pressed state for the current event.  This can 
-	* be useful to turn off in cases where you always want a full press to override the light press, even 
-	* if it was a very slow press
-	*
-	* @param	bAllow		Enables or disables light press locking for this hand
-	*/
+	 * Enables or disable "light press" locking for this interactor.  This is an advanced feature.
+	 * It should be true if we allow locking of the lightly pressed state for the current event.  This can 
+	 * be useful to turn off in cases where you always want a full press to override the light press, even 
+	 * if it was a very slow press
+	 *
+	 * @param	bAllow		Enables or disables light press locking for this hand
+	 */
 	void SetAllowTriggerLightPressLocking( const bool bInAllow );
 
-	/** Check if this interactor allows the light press to be locked */
+	/** @return Check if this interactor allows the light press to be locked */
 	bool AllowTriggerLightPressLocking() const;
 
+	/**
+	 * When responding to a "light press" event, you can call this and pass false to to prevent a full press event from
+	 * being possible at all.  Useful when you're not planning to use FullyPressed for anything, but you don't want a 
+	 * full press to cancel your light press.
+	 *
+	 * @param	bInAllow	Whether to allow a full press to interrupt the light press.  Defaults to true, and will reset to true with every re-press of the trigger.
+	 */
+	void SetAllowTriggerFullPress( const bool bInAllow );
+
+	/** @return Returns whether a full press is allowed to interrupt a light press. */
+	bool AllowTriggerFullPress() const;
+
+
 protected:
 
 	/** To be overridden by base class. Called by HandleInputKey before delegates and default input implementation */
diff --git a/Engine/Source/Editor/ViewportInteraction/Public/ViewportWorldInteraction.h b/Engine/Source/Editor/ViewportInteraction/Public/ViewportWorldInteraction.h
--- a/Engine/Source/Editor/ViewportInteraction/Public/ViewportWorldInteraction.h
+++ b/Engine/Source/Editor/ViewportInteraction/Public/ViewportWorldInteraction.h
@@ -179,6 +179,12 @@ public:
 	/** Gets the transform gizmo actor, or returns null if we currently don't have one */
 	class ABaseTransformGizmo* GetTransformGizmoActor();
 
+	/** Sets whether the transform gizmo should be visible at all */
+	void SetTransformGizmoVisible( const bool bShouldBeVisible );
+
+	/** Returns whether the transform gizmo is visible */
+	bool IsTransformGizmoVisible() const;
+
 	/** Sets if objects are dragged with either hand since last time selecting something */
 	void SetDraggedSinceLastSelection( const bool bNewDraggedSinceLastSelection );
 
@@ -448,6 +454,9 @@ private:
 	/** The current Gizmo type that is used for the TransformGizmo Actor */
 	EGizmoHandleTypes CurrentGizmoType;
 
+	/** Whether the gizmo should be visible or not */
+	bool bIsTransformGizmoVisible;
+
 
 	//
 	// Snap grid
diff --git a/Engine/Source/Editor/ViewportInteraction/ViewportInteractor.cpp b/Engine/Source/Editor/ViewportInteraction/ViewportInteractor.cpp
--- a/Engine/Source/Editor/ViewportInteraction/ViewportInteractor.cpp
+++ b/Engine/Source/Editor/ViewportInteraction/ViewportInteractor.cpp
@@ -265,7 +265,7 @@ bool UViewportInteractor::HandleInputKey( const FKey Key, const EInputEvent Even
 
 										// If we've selected an actor with a light press, don't prevent the user from moving the actor if they
 										// press down all the way.  We only want locking when interacting with gizmos
-										InteractorData.bAllowTriggerLightPressLocking = false;
+										SetAllowTriggerLightPressLocking( false );
 
 										// If the user did a full press to select an actor that wasn't selected, allow them to drag it right away
 										const bool bOtherHandTryingToDrag =
@@ -594,4 +594,14 @@ void UViewportInteractor::SetAllowTriggerLightPressLocking( const bool bInAllow
 bool UViewportInteractor::AllowTriggerLightPressLocking() const
 {
 	return InteractorData.bAllowTriggerLightPressLocking;
+}
+
+void UViewportInteractor::SetAllowTriggerFullPress( const bool bInAllow )
+{
+	InteractorData.bAllowTriggerFullPress = bInAllow;
+}
+
+bool UViewportInteractor::AllowTriggerFullPress() const
+{
+	return InteractorData.bAllowTriggerFullPress;
 }
\ No newline at end of file
diff --git a/Engine/Source/Editor/ViewportInteraction/ViewportInteractorData.h b/Engine/Source/Editor/ViewportInteraction/ViewportInteractorData.h
--- a/Engine/Source/Editor/ViewportInteraction/ViewportInteractorData.h
+++ b/Engine/Source/Editor/ViewportInteraction/ViewportInteractorData.h
@@ -31,9 +31,6 @@ struct FViewportInteractorData
 	/** Lights up when our laser pointer is close enough to something to click */
 	bool bIsHovering;
 
-	/** True if we're hovering over UI right now.  When hovering over UI, we don't bother drawing a see-through laser pointer */
-	bool bIsHoveringOverUI; //@todo: ViewportInteraction: UI should not be in this module
-
 	/** The widget component we are currently hovering over */
 	class UWidgetComponent* HoveringOverWidgetComponent; //@todo: ViewportInteraction: UI should not be in this module
 
@@ -131,11 +128,15 @@ struct FViewportInteractorData
 	a very slow press */
 	bool bAllowTriggerLightPressLocking;
 
+	/** This can be set to false when handling SelectAndMove_LightlyPressed to prevent a full press event from
+	    being possible at all.  Useful when you're not planning to use FullyPressed for anything, but you
+		don't want a full press to cancel your light press. */
+	bool bAllowTriggerFullPress;
+
 	/** Default constructor for FVirtualHand that initializes safe defaults */
 	FViewportInteractorData()
 	{
 		bIsHovering = false;
-		bIsHoveringOverUI = false;
 		HoveringOverWidgetComponent = nullptr;
 		HoverLocation = FVector::ZeroVector;
 		HoveredActorComponent = nullptr;
@@ -165,6 +166,7 @@ struct FViewportInteractorData
 		HoverHapticCheckLastHoveredGizmoComponent = nullptr;
 
 		bAllowTriggerLightPressLocking = true;
+		bAllowTriggerFullPress = true;
 	}
 };
 
diff --git a/Engine/Source/Editor/ViewportInteraction/ViewportWorldInteraction.cpp b/Engine/Source/Editor/ViewportInteraction/ViewportWorldInteraction.cpp
--- a/Engine/Source/Editor/ViewportInteraction/ViewportWorldInteraction.cpp
+++ b/Engine/Source/Editor/ViewportInteraction/ViewportWorldInteraction.cpp
@@ -228,6 +228,7 @@ UViewportWorldInteraction::UViewportWorldInteraction( const FObjectInitializer&
 	StartDragAngleOnRotation( ),
 	StartDragHandleDirection( ),
 	CurrentGizmoType( EGizmoHandleTypes::All ),
+	bIsTransformGizmoVisible( true ),
 	SnapGridActor( nullptr ),
 	SnapGridMeshComponent( nullptr ),
 	SnapGridMID( nullptr ),
@@ -2620,13 +2621,27 @@ void UViewportWorldInteraction::SpawnTransformGizmoIfNeeded()
 		check( TransformGizmoActor != nullptr );
 		TransformGizmoActor->SetOwnerWorldInteraction( this );
 
-		if ( VI::ShowTransformGizmo->GetInt() == 0 )
+		if ( VI::ShowTransformGizmo->GetInt() == 0 || !bIsTransformGizmoVisible )
 		{
 			TransformGizmoActor->SetIsTemporarilyHiddenInEditor( true );
 		}
 	}
 }
 
+void UViewportWorldInteraction::SetTransformGizmoVisible( const bool bShouldBeVisible )
+{
+	bIsTransformGizmoVisible = bShouldBeVisible;
+	if( TransformGizmoActor != nullptr )
+	{
+		TransformGizmoActor->SetIsTemporarilyHiddenInEditor( VI::ShowTransformGizmo->GetInt() == 0 || !bIsTransformGizmoVisible );
+	}
+}
+
+bool UViewportWorldInteraction::IsTransformGizmoVisible() const
+{
+	return bIsTransformGizmoVisible;
+}
+
 void UViewportWorldInteraction::ApplyVelocityDamping( FVector& Velocity, const bool bVelocitySensitive )
 {
 	const float InertialMovementZeroEpsilon = 0.01f;	// @todo vreditor tweak
