Copying //UE4/Dev-VREditor to //UE4/Dev-Main (Source: //UE4/Dev-VREditor @ 3051763)

==========================
MAJOR FEATURES + CHANGES
==========================

Change 2951403 on 2016/04/21 by [REDACTED_PERSON]

	VREditor : Pivot Transform Gizmo first version
	- Switch to old gizmo with command: VREd.PivotPointTransformGizmo 0

Change 2953179 on 2016/04/22 by [REDACTED_PERSON]

	Adding flashlight functionality to the VR Editor quick menu. The "Light" button on the quick menu will add a flashlight to the hand without the menu attached.

Change 2953438 on 2016/04/22 by [REDACTED_PERSON]

	VR Editor :
	- Fix mac warning for WorldInteraction constructor list order
	- Fix mac warning for UniversalGizmoHandle using the same variable name

Change 2956455 on 2016/04/26 by [REDACTED_PERSON]

	Merging //UE4/Dev-Main to Dev-VREditor (//UE4/Dev-VREditor)

Change 2960831 on 2016/04/29 by [REDACTED_PERSON]

	Add stream settings for VREditor branch which shows GitHub on EC dashboard.

Change 2977191 on 2016/05/13 by [REDACTED_PERSON]

	VR Editor: Adding screenshot button to the VR Editor Quick Menu. The screenshot does undistorted left eye + right eye, and does not capture UI.

Change 2985532 on 2016/05/20 by [REDACTED_PERSON]

	VR Editor: Fix for UE-27727, unable to undo a duplication

Change 2985585 on 2016/05/20 by [REDACTED_PERSON]

	VR Editor: Fixed asset editor UI popping up whenever a Modes icon was clicked in VR

Change 2985586 on 2016/05/20 by [REDACTED_PERSON]

	VR Editor: Vertex and texture painting support (UETOOL-668)
	- You can now activate Mesh Paint and use the vertex/texture painting while in VR!
	- Some features are not yet accessible (e.g. color picker)
	- There is a known issue where the HMD view can flicker while painting on a texture

Change 2992730 on 2016/05/27 by [REDACTED_PERSON]

	VR Editor: Adding access to the World Settings UI from the Quick Menu (UETOOL-722). Temp. quick menu button reorg.

Change 2999582 on 2016/06/03 by [REDACTED_PERSON]

	VR Editor: Fixing small snap buttons due to adding a uniform grid panel.

Change 3001771 on 2016/06/06 by [REDACTED_PERSON]

	VR Editor: Adding return values to the quick menu functions for more discrete button refreshes

Change 3004009 on 2016/06/07 by [REDACTED_PERSON]

	VR Editor: Changing doubleclick event to be set on pressed instead of released. UETOOL-878

Change 3012664 on 2016/06/14 by [REDACTED_PERSON]

	VR Editor:
	- First extensibility stage for interaction
	- Fixed world space rotation with gizmo handle

Change 3013040 on 2016/06/14 by [REDACTED_PERSON]

	VR Editor: - Removed AutoConsoleVariable duplicates
	- Fixed mac errors: declarations that shadows another field
	- Fixed mac errors: Constructor initialization list order

Change 3013400 on 2016/06/14 by [REDACTED_PERSON]

	VR Editor: - Fix mac warning: declaration shadows a field for SnapGridActor
	- Fix SetHasUIOnForearm setter for bHasUIOnForearm

Change 3019658 on 2016/06/20 by [REDACTED_PERSON]

	Crash when setting force feedback while SteamController plugin is enabled. This was caused by setting a force feedback to a SteamController device while none has been initialized (I belive the plugin enabled this). Solution is to check whether a gamepad of that type is attached before doing this. It's possible that we'll have to further implement IsGamepadAttached in other devices to conform with IInputDevice interface.

Change 3024534 on 2016/06/23 by [REDACTED_PERSON]

	VR Editor: Miscellaneous fixes (Merged from CL 3024520)
	- Fixed input capture not being revoked properly after releasing controller triggers
	- Fixed some shadowed variable warnings
	- Renamed some function parameter variables

Change 3024713 on 2016/06/23 by [REDACTED_PERSON]

	Repplying "Undoing Duplication" fix #jira UE-27727

Change 3025030 on 2016/06/23 by [REDACTED_PERSON]

	VR Editor:
	- Removed InputEvent parameter from the ViewportFOnVIActionHandle, because this data was already stored in the FViewportActionKeyInput that is also passed
	- Fixed typo ViewportDragOperation SetDragOperationClass
	- Removed OnDrag from interactables
	- Fixed mac warning with order of ViewportWorldInteraction constructor initialization list
	- Added events for AVREditorButton
	- Added Switching between VR Editor to VR PIE:
	  - Changed quickmenu to add PIE in VR button
	  - Module starts PIE when the reason for exiting the mode was to start PIE session
	  - Module sets the initial worldtometers for the VR Editor and stores it, so when exiting the VR Editor mode the worldtometers is set back from what it was before VR Editor.

Change 3026518 on 2016/06/24 by [REDACTED_PERSON]

	Merging //UE4/Dev-Main to Dev-VREditor (//UE4/Dev-VREditor)

Change 3026577 on 2016/06/24 by [REDACTED_PERSON]

	Viewport Interaction: Mouse cursor support (merged from CL 3026507)
	- Added a new 'mouse cursor' interactor that can be enabled (when not in VR)
	- We now automatically create a "mouse cursor" interactor when running in 'forced VR mode', so you can click on objects and gizmos!
	- Currently it only supports 'SelectAndMove' (LMB), 'SelectAndMove_LightlyPressed' (MMB) and 'WorldMovement' (RMB)
	- Interacting with UI using the mouse cursor is not supported yet.

	Other fixes:
	- ViewportInteraction: Fixed assertions when an interactor had no 'pair' and certain actions were triggered
	- ViewportInteractor: Now creates a default transform for the interactor if no overridden implementation was provided
	- VREditorInteractor: Fixed bad memory references to expired stack addresses
	- VREditor: Fixed mouse buttons events not being routed through to HandleInputKey() (preprocessor did not support those.)
	- VREditor: External modules can now activate VR Editor in 'forced VR mode', using new optional parameter to EnableVREditor()
	- VREditor: Tutorial no longer shows up in 'forced VR mode'.  You can't click on UI anyway, and it was just in the way

Change 3026786 on 2016/06/24 by [REDACTED_PERSON]

	VREditor: Fixed assertion when hovering over model components (BSP)

Change 3027015 on 2016/06/24 by [REDACTED_PERSON]

	VR Editor: Motion controllers are now enabled in forced VR mode (merged from CL 3027012)
	(Console command is 'VREd.ForceVRMode')

Change 3029346 on 2016/06/27 by [REDACTED_PERSON]

	VREditor: - Moved ViewportMotionControllerInteractor to VREditor module to solve inheritance issue
	- Fixed laser visibility (UE-32529 & UE-32523)
	- Fixed clearing delegate in Teleport

Change 3032909 on 2016/06/29 by [REDACTED_PERSON]

	Fixed a bug in FSteamVRController::IsGamepadAttached not properly checking for right hand tracking status.

Change 3033446 on 2016/06/29 by [REDACTED_PERSON]

Change 3033487 on 2016/06/29 by [REDACTED_PERSON]

Change 3038558 on 2016/07/05 by [REDACTED_PERSON]

	Fixed world scaling for hover point while transforming in VR editor

Change 3039896 on 2016/07/06 by [REDACTED_PERSON]

	Auto-entry into VR editing mode when a Vive or Oculus Rift is worn. Editor setting under Experimental to toggle the behavior, new enum EHMDWornState which defaults to Unknown for headsets that don't have sensors to detect user proximity/interaction.

Change 3042038 on 2016/07/07 by [REDACTED_PERSON]

	Small clean up for MeshPaintEd

Change 3045002 on 2016/07/11 by [REDACTED_PERSON]

Change 3045059 on 2016/07/11 by [REDACTED_PERSON]

	[FoliageEdMode] Renamed variable to not hide declaration

Change 3048240 on 2016/07/13 by [REDACTED_PERSON]

	Removed redundant code in FoliageEdMode

Change 3049839 on 2016/07/14 by [REDACTED_PERSON]

	Merging //UE4/Dev-Main to Dev-VREditor (//UE4/Dev-VREditor) (from CL 3049614)

Change 3049848 on 2016/07/14 by [REDACTED_PERSON]

	VREditor: Fixed foliage fill not working

Change 3050395 on 2016/07/14 by [REDACTED_PERSON]

	Removed leftover debug code #rb mike.fricker

Change 3050444 on 2016/07/14 by [REDACTED_PERSON]

	Fixed laser not re-appearing after closing a docked UI panel (UE-33299)

Change 3050669 on 2016/07/14 by [REDACTED_PERSON]

	VREditor: Fixed foliage select/lasso selection highlight not shown/updated

Change 3051407 on 2016/07/15 by [REDACTED_PERSON]

	VR Editor: Resaved some assets with a versioned editor (no effective changes)

Change 3051577 on 7/15/2016 by [REDACTED_PERSON]

	VR Editor: Fixed VR Editor assertion in server builds

Change 3051763 on 7/15/2016 by [REDACTED_PERSON]

	VR Editor: Avoid loading VR Editor classes in the cooker (UE-33318)

diff --git a/Engine/Source/Editor/MeshPaint/Private/MeshPaintEdMode.cpp b/Engine/Source/Editor/MeshPaint/Private/MeshPaintEdMode.cpp
--- a/Engine/Source/Editor/MeshPaint/Private/MeshPaintEdMode.cpp
+++ b/Engine/Source/Editor/MeshPaint/Private/MeshPaintEdMode.cpp
@@ -20,7 +20,6 @@
 #include "Editor/LevelEditor/Public/SLevelViewport.h"
 #include "MessageLog.h"
 
-
 #include "Runtime/Engine/Classes/PhysicsEngine/BodySetup.h"
 #include "SMeshPaint.h"
 #include "ComponentReregisterContext.h"
@@ -33,6 +32,9 @@
 #include "MeshPaintAdapterFactory.h"
 #include "Components/SplineMeshComponent.h"
 
+#include "ViewportWorldInteraction.h"
+#include "VREditorInteractor.h"
+
 #define LOCTEXT_NAMESPACE "MeshPaint_Mode"
 
 DEFINE_LOG_CATEGORY_STATIC(LogMeshPaintEdMode, Log, All);
@@ -88,6 +90,7 @@ public:
 FEdModeMeshPaint::FEdModeMeshPaint() 
 	: FEdMode(),
 	  bIsPainting( false ),
+	  PaintingWithInteractorInVR( nullptr ),
 	  bIsFloodFill( false ),
 	  bPushInstanceColorsToMesh( false ),
 	  PaintingStartTime( 0.0 ),
@@ -190,18 +193,39 @@ void FEdModeMeshPaint::Enter()
 		SetViewportShowFlags( bAllowColorViewModes, *ViewportClient );
 	} 
 
-	//When painting vertext colors we want to force the lod level of objects being painted to LOD0.
+	//When painting vertex colors we want to force the lod level of objects being painted to LOD0.
 	if( FMeshPaintSettings::Get().ResourceType == EMeshPaintResource::VertexColors )
 	{
 		ApplyOrRemoveForceBestLOD(/*bApply=*/ true);
 	}
-}
 
+	// Register to find out about editor modes being activated
+	GetModeManager()->OnEditorModeChanged().AddRaw( this, &FEdModeMeshPaint::OnEditorModeChanged );
 
+	// Register to find out about VR input events
+	IVREditorMode* VREditorMode = static_cast<IVREditorMode*>( GetModeManager()->GetActiveMode( IVREditorModule::Get().GetVREditorModeID() ) );
+	if( VREditorMode != nullptr )
+	{
+		VREditorMode->GetWorldInteraction().OnViewportInteractionInputAction().RemoveAll( this );
+		VREditorMode->GetWorldInteraction().OnViewportInteractionInputAction().AddRaw( this, &FEdModeMeshPaint::OnVRAction );
+	}
+}
 
 /** FEdMode: Called when the mode is exited */
 void FEdModeMeshPaint::Exit()
 {
+	// Unregister from event handlers
+	if( IVREditorModule::IsAvailable() )
+	{
+		IVREditorMode* VREditorMode = static_cast<IVREditorMode*>( GetModeManager()->GetActiveMode( IVREditorModule::Get().GetVREditorModeID() ) );
+		if( VREditorMode != nullptr )
+		{
+			VREditorMode->GetWorldInteraction().OnViewportInteractionInputAction().RemoveAll( this );
+		}
+	}
+
+	GetModeManager()->OnEditorModeChanged().RemoveAll( this );
+
 	//If we're painting vertex colors then propagate the painting done on LOD0 to all lower LODs. 
 	//Then stop forcing the LOD level of the mesh to LOD0.
 	if( FMeshPaintSettings::Get().ResourceType == EMeshPaintResource::VertexColors )
@@ -271,6 +295,23 @@ void FEdModeMeshPaint::Exit()
 }
 
 
+void FEdModeMeshPaint::OnEditorModeChanged( FEdMode* EditorMode, bool bEntered )
+{
+	// VR Editor mode may have gone away, so re-register for events in case it's a different object than the one we were originally bound to
+	check( EditorMode != nullptr );
+	if( bEntered && EditorMode->GetID() == IVREditorModule::Get().GetVREditorModeID() )
+	{
+		IVREditorMode* VREditorMode = static_cast<IVREditorMode*>( EditorMode );
+		if ( VREditorMode != nullptr )
+		{
+			VREditorMode->GetWorldInteraction().OnViewportInteractionInputAction().RemoveAll( this );
+			VREditorMode->GetWorldInteraction().OnViewportInteractionInputAction().AddRaw( this, &FEdModeMeshPaint::OnVRAction );
+		}
+
+		PaintingWithInteractorInVR = nullptr;
+	}
+}
+
 
 /** FEdMode: Called when the mouse is moved over the viewport */
 bool FEdModeMeshPaint::MouseMove( FEditorViewportClient* ViewportClient, FViewport* Viewport, int32 x, int32 y )
@@ -299,7 +340,7 @@ bool FEdModeMeshPaint::MouseMove( FEditorViewportClient* ViewportClient, FViewpo
 bool FEdModeMeshPaint::CapturedMouseMove( FEditorViewportClient* InViewportClient, FViewport* InViewport, int32 InMouseX, int32 InMouseY )
 {
 	// We only care about perspective viewports
-	if( InViewportClient->IsPerspective() && InViewportClient->EngineShowFlags.ModeWidgets )
+	if( InViewportClient->IsPerspective() )
 	{
 		if( bIsPainting )
 		{
@@ -316,12 +357,13 @@ bool FEdModeMeshPaint::CapturedMouseMove( FEditorViewportClient* InViewportClien
 			// Paint!
 			const bool bVisualCueOnly = false;
 			const EMeshPaintAction::Type PaintAction = GetPaintAction(InViewport);
+
 			// Apply stylus pressure
 			const float StrengthScale = InViewport->IsPenActive() ? InViewport->GetTabletPressure() : 1.f;
 
-			bool bAnyPaintAbleActorsUnderCursor = false;
+			bool bAnyPaintableActorsUnderCursor = false;
 
-			DoPaint( View->ViewMatrices.ViewOrigin, MouseViewportRay.GetOrigin(), MouseViewportRay.GetDirection(), NULL, PaintAction, bVisualCueOnly, StrengthScale, bAnyPaintAbleActorsUnderCursor );
+			DoPaint( View->ViewMatrices.ViewOrigin, MouseViewportRay.GetOrigin(), MouseViewportRay.GetDirection(), NULL, PaintAction, bVisualCueOnly, StrengthScale, /* Out */ bAnyPaintableActorsUnderCursor );
 
 			return true;
 		}
@@ -361,6 +403,7 @@ void FEdModeMeshPaint::EndPainting()
 	if(bIsPainting)
 	{
 		bIsPainting = false;
+		PaintingWithInteractorInVR = nullptr;
 		FinishPaintingTexture();
 
 		// Rebuild any static meshes that we painted on last stroke
@@ -459,15 +502,16 @@ bool FEdModeMeshPaint::InputKey( FEditorViewportClient* InViewportClient, FViewp
 		}
 	}
 
-	// When painting we only care about perspective viewports where we are we are allowed to show mode widgets
-	if( !bIsAltDown && InViewportClient->IsPerspective() && InViewportClient->EngineShowFlags.ModeWidgets)
+	// When painting we only care about perspective viewports
+	if( !bIsAltDown && InViewportClient->IsPerspective() )
 	{
 		// Does the user want to paint right now?
 		const bool bUserWantsPaint = bIsLeftButtonDown && !bIsRightButtonDown && !bIsAltDown;
-		bool bAnyPaintAbleActorsUnderCursor = false;
+		bool bAnyPaintableActorsUnderCursor = false;
 
 		// Stop current tracking if the user is no longer painting
-		if( bIsPainting && !bUserWantsPaint )
+		if( bIsPainting && !bUserWantsPaint &&
+			( InKey == EKeys::LeftMouseButton || InKey == EKeys::RightMouseButton && InKey == EKeys::LeftAlt || InKey == EKeys::RightAlt ) )
 		{
 			bHandled = true;
 			EndPainting();
@@ -506,7 +550,7 @@ bool FEdModeMeshPaint::InputKey( FEditorViewportClient* InViewportClient, FViewp
 						const bool bVisualCueOnly = false;
 						const EMeshPaintAction::Type PaintAction = GetPaintAction(InViewport);
 						const float StrengthScale = 1.0f;
-						DoPaint( View->ViewMatrices.ViewOrigin, MouseViewportRay.GetOrigin(), MouseViewportRay.GetDirection(), NULL, PaintAction, bVisualCueOnly, StrengthScale, bAnyPaintAbleActorsUnderCursor );
+						DoPaint( View->ViewMatrices.ViewOrigin, MouseViewportRay.GetOrigin(), MouseViewportRay.GetDirection(), NULL, PaintAction, bVisualCueOnly, StrengthScale, /* Out */ bAnyPaintableActorsUnderCursor );
 					}
 				}
 				else
@@ -518,7 +562,7 @@ bool FEdModeMeshPaint::InputKey( FEditorViewportClient* InViewportClient, FViewp
 			}
 		}
 
-		if( !bAnyPaintAbleActorsUnderCursor )
+		if( !bAnyPaintableActorsUnderCursor )
 		{
 			bHandled = false;
 		}
@@ -1007,7 +1051,7 @@ void FEdModeMeshPaint::DoPaint( const FVector& InCameraOrigin,
 								const EMeshPaintAction::Type InPaintAction,
 								const bool bVisualCueOnly,
 								const float InStrengthScale,
-								OUT bool& bAnyPaintAbleActorsUnderCursor)
+								OUT bool& bAnyPaintableActorsUnderCursor)
 {
 	const float BrushRadius = GetBrushRadiiDefault();
 
@@ -1089,10 +1133,10 @@ void FEdModeMeshPaint::DoPaint( const FVector& InCameraOrigin,
 		}
 	}
 
-	bAnyPaintAbleActorsUnderCursor = (PaintableComponents.Num() > 0);
+	bAnyPaintableActorsUnderCursor = (PaintableComponents.Num() > 0);
 
 	// Are we actually applying paint here?
-	const bool bShouldApplyPaint = bAnyPaintAbleActorsUnderCursor && ( (bIsPainting && !bVisualCueOnly) || 
+	const bool bShouldApplyPaint = bAnyPaintableActorsUnderCursor && ( (bIsPainting && !bVisualCueOnly) || 
 		(InPaintAction == EMeshPaintAction::Fill) || 
 		(InPaintAction == EMeshPaintAction::PushInstanceColorsToMesh) );
 
@@ -2573,12 +2617,8 @@ void FEdModeMeshPaint::Render( const FSceneView* View, FViewport* Viewport, FPri
 	/** Call parent implementation */
 	FEdMode::Render( View, Viewport, PDI );
 
-	// If this viewport does not support Mode widgets we will not draw it here.
 	FEditorViewportClient* ViewportClient = (FEditorViewportClient*)Viewport->GetClient();
-	if( ViewportClient && !ViewportClient->EngineShowFlags.ModeWidgets)
-	{
-		return;
-	}
+	check( ViewportClient != nullptr );
 
 	// We only care about perspective viewports
 	const bool bIsPerspectiveViewport = ( View->ViewMatrices.ProjMatrix.M[ 3 ][ 3 ] < ( 1.0f - SMALL_NUMBER ) );
@@ -2594,37 +2634,95 @@ void FEdModeMeshPaint::Render( const FSceneView* View, FViewport* Viewport, FPri
 		const bool bAllowColorViewModes = ( FMeshPaintSettings::Get().ResourceType != EMeshPaintResource::Texture );
 		SetViewportShowFlags( bAllowColorViewModes, *ViewportClient );
 
-		// Make sure the cursor is visible OR we're flood filling.  No point drawing a paint cue when there's no cursor.
-		if( Viewport->IsCursorVisible() || IsForceRendered())
-		{
 
-			if( !PDI->IsHitTesting() )
-			{
-				// Grab the mouse cursor position
-				FIntPoint MousePosition;
-				Viewport->GetMousePos( MousePosition );
+		bool bHavePaintRay = false;
+		FVector PaintRayViewOrigin = FVector::ZeroVector;
+		FVector PaintRayStart = FVector::ZeroVector;
+		FVector PaintRayDirection = FVector::ZeroVector;
 
-				// Is the mouse currently over the viewport? or flood filling
-				if(IsForceRendered() || ( MousePosition.X >= 0 && MousePosition.Y >= 0 && MousePosition.X < (int32)Viewport->GetSizeXY().X && MousePosition.Y < (int32)Viewport->GetSizeXY().Y) )
+		// Check to see if VREditorMode is active
+		if ( IVREditorModule::IsAvailable() )
+		{
+			IVREditorMode* VREditorMode = static_cast< IVREditorMode* >( GetModeManager()->GetActiveMode( IVREditorModule::Get().GetVREditorModeID() ) );
+			if ( VREditorMode != nullptr && VREditorMode->IsFullyInitialized() )
+			{
+				for ( UViewportInteractor* Interactor : VREditorMode->GetWorldInteraction().GetInteractors() )
 				{
-					// Compute a world space ray from the screen space mouse coordinates
-					FViewportCursorLocation MouseViewportRay( View, ViewportClient, MousePosition.X, MousePosition.Y );
+					// Skip other interactors if we are painting with one
+					if ( PaintingWithInteractorInVR && Interactor != PaintingWithInteractorInVR )
+					{
+						continue;
+					}
 
+					FVector LaserPointerStart, LaserPointerEnd;
+					if ( Interactor->GetLaserPointer( /* Out */ LaserPointerStart, /* Out */ LaserPointerEnd ) )
+					{
+						const FVector LaserPointerDirection = ( LaserPointerEnd - LaserPointerStart ).GetSafeNormal();
 
-					// Unless "Flow" mode is enabled, we'll only draw a visual cue while rendering and won't
-					// do any actual painting.  When "Flow" is turned on we will paint here, too!
-					const bool bVisualCueOnly = !FMeshPaintSettings::Get().bEnableFlow;
-					float StrengthScale = FMeshPaintSettings::Get().bEnableFlow ? FMeshPaintSettings::Get().FlowAmount : 1.0f;
+						bHavePaintRay = true;
+						PaintRayViewOrigin = VREditorMode->GetHeadTransform().GetLocation();
+						PaintRayStart = LaserPointerStart;
+						PaintRayDirection = LaserPointerDirection;
+					}
 
-					// Apply stylus pressure if it's active
-					if( Viewport->IsPenActive() )
+					if ( !bHavePaintRay )
 					{
-						StrengthScale *= Viewport->GetTabletPressure();
+						// Make sure the cursor is visible OR we're flood filling.  No point drawing a paint cue when there's no cursor.
+						if ( Viewport->IsCursorVisible() || IsForceRendered() )
+						{
+							if ( !PDI->IsHitTesting() )
+							{
+								// Grab the mouse cursor position
+								FIntPoint MousePosition;
+								Viewport->GetMousePos( MousePosition );
+
+								// Is the mouse currently over the viewport? or flood filling
+								if ( IsForceRendered() || ( MousePosition.X >= 0 && MousePosition.Y >= 0 && MousePosition.X < ( int32 ) Viewport->GetSizeXY().X && MousePosition.Y < ( int32 ) Viewport->GetSizeXY().Y ) )
+								{
+									// Compute a world space ray from the screen space mouse coordinates
+									FViewportCursorLocation MouseViewportRay( View, ViewportClient, MousePosition.X, MousePosition.Y );
+
+									bHavePaintRay = true;
+									PaintRayViewOrigin = View->ViewMatrices.ViewOrigin;
+									PaintRayStart = MouseViewportRay.GetOrigin();
+									PaintRayDirection = MouseViewportRay.GetDirection();
+								}
+							}
+						}
 					}
 
-					const EMeshPaintAction::Type PaintAction = GetPaintAction(Viewport);
-					bool bAnyPaintAbleActorsUnderCursor = false;
-					DoPaint( View->ViewMatrices.ViewOrigin, MouseViewportRay.GetOrigin(), MouseViewportRay.GetDirection(), PDI, PaintAction, bVisualCueOnly, StrengthScale, bAnyPaintAbleActorsUnderCursor );
+					if ( bHavePaintRay )
+					{
+						// Unless "Flow" mode is enabled, we'll only draw a visual cue while rendering and won't
+						// do any actual painting.  When "Flow" is turned on we will paint here, too!
+						const bool bVisualCueOnly = !bIsPainting || ( PaintingWithInteractorInVR == nullptr && !FMeshPaintSettings::Get().bEnableFlow );
+						float StrengthScale = ( PaintingWithInteractorInVR == nullptr && FMeshPaintSettings::Get().bEnableFlow ) ? FMeshPaintSettings::Get().FlowAmount : 1.0f;
+
+						// Apply VR controller trigger pressure if we're painting in VR
+						if ( bIsPainting && PaintingWithInteractorInVR && VREditorMode != nullptr )
+						{
+							UVREditorInteractor* VRInteractor = Cast<UVREditorInteractor>( PaintingWithInteractorInVR );
+							if ( VRInteractor )
+							{
+								StrengthScale *= VRInteractor->GetSelectAndMoveTriggerValue();
+
+								// Make sure light press locking is enabled in VR for the hand we're painting with.  We don't want
+								// a full press to be converted to a move action, however locking will have been disabled already
+								// by the world interaction code after a light press happened on an object. 
+								VRInteractor->SetAllowTriggerLightPressLocking( true );
+							}
+						}
+
+						// Apply stylus pressure if it's active
+						else if ( Viewport->IsPenActive() )
+						{
+							StrengthScale *= Viewport->GetTabletPressure();
+						}
+
+						const EMeshPaintAction::Type PaintAction = GetPaintAction( Viewport );
+						bool bAnyPaintableActorsUnderCursor = false;
+						DoPaint( PaintRayViewOrigin, PaintRayStart, PaintRayDirection, PDI, PaintAction, bVisualCueOnly, StrengthScale, /* Out */ bAnyPaintableActorsUnderCursor );
+					}
 				}
 			}
 		}
@@ -3418,6 +3516,18 @@ EMeshPaintAction::Type FEdModeMeshPaint::GetPaintAction(FViewport* InViewport)
 	else
 	{
 		PaintAction = bShiftDown ? EMeshPaintAction::Erase : EMeshPaintAction::Paint;
+
+		// When painting using VR, allow the modifier button to activate Erase mode
+		if( PaintingWithInteractorInVR )
+		{
+			IVREditorMode* VREditorMode = static_cast<IVREditorMode*>( GetModeManager()->GetActiveMode( IVREditorModule::Get().GetVREditorModeID() ) );
+			UVREditorInteractor* VRInteractor = Cast<UVREditorInteractor>( PaintingWithInteractorInVR );
+			if( VREditorMode != nullptr && VRInteractor != nullptr )
+			{
+				const bool bIsModifierPressed = VRInteractor->IsModifierPressed();
+				PaintAction = bIsModifierPressed ? EMeshPaintAction::Erase : EMeshPaintAction::Paint;
+			}
+		}
 	}
 	return PaintAction;
 
@@ -4966,4 +5076,66 @@ TWeakObjectPtr<AActor> FEdModeMeshPaint::GetEditingActor() const
 }
 
 
+void FEdModeMeshPaint::OnVRAction( class FEditorViewportClient& ViewportClient, UViewportInteractor* Interactor, 
+	const FViewportActionKeyInput& Action, bool& bOutIsInputCaptured, bool& bWasHandled )
+{
+	IVREditorMode* VREditorMode = static_cast<IVREditorMode*>( GetModeManager()->GetActiveMode( IVREditorModule::Get().GetVREditorModeID() ) );
+	UVREditorInteractor* VRInteractor = Cast<UVREditorInteractor>( Interactor );
+
+	if( VREditorMode != nullptr && VRInteractor != nullptr )
+	{
+		// Absorb all "full press" events unless the user is clicking on UI to avoid objects getting moved around.
+		// @todo vreditor: We shouldn't have to use this trick.  This input system needs more flexibility in its public interface!
+		if( Action.ActionType == ViewportWorldActionTypes::SelectAndMove && bIsPainting && 
+			PaintingWithInteractorInVR && PaintingWithInteractorInVR == Interactor && 
+			!VRInteractor->IsHoveringOverUI() )
+		{
+			bWasHandled = true;
+			bOutIsInputCaptured = true;
+		}
+
+		if( Action.ActionType == ViewportWorldActionTypes::SelectAndMove_LightlyPressed )
+		{
+			if( !bIsPainting && Action.Event == IE_Pressed && !VRInteractor->IsHoveringOverUI() )
+			{
+				// NOTE: We intentionally don't capture input for "LightlyPressed" because we want to allow objects to still
+				// become selected when clicked on.  We avoid allowing the objects to be moved by calling the 
+				// SetAllowTriggerLightPressLocking() function in the next tick (after VR Editor will have disabled locking.)
+				// @todo vreditor: We shouldn't have to use this trick.  This input system needs more flexibility in its public interface!
+				//bWasHandled = true;
+				//bOutIsInputCaptured = true;
+
+				StartPainting();
+				PaintingWithInteractorInVR = Interactor;
+
+				// Go ahead and paint immediately
+				FVector LaserPointerStart, LaserPointerEnd;
+				if( Interactor->GetLaserPointer( /* Out */ LaserPointerStart, /* Out */ LaserPointerEnd ) )
+				{
+					const FVector LaserPointerDirection = ( LaserPointerEnd - LaserPointerStart ).GetSafeNormal();
+
+					// Apply VR controller trigger pressure
+					float StrengthScale = VRInteractor->GetSelectAndMoveTriggerValue();
+
+					// Paint!
+					const bool bVisualCueOnly = false;
+					const EMeshPaintAction::Type PaintAction = GetPaintAction( ViewportClient.Viewport );
+					bool bAnyPaintableActorsUnderCursor = false;
+					DoPaint( VREditorMode->GetHeadTransform().GetLocation(), LaserPointerStart, LaserPointerDirection, NULL, PaintAction, bVisualCueOnly, StrengthScale, /* Out */ bAnyPaintableActorsUnderCursor );
+				}
+			}
+
+			// Stop current tracking if the user is no longer painting
+			else if( bIsPainting && Action.Event == IE_Released && PaintingWithInteractorInVR && PaintingWithInteractorInVR == Interactor )
+			{
+				EndPainting();
+
+				// @todo vreditor: We shouldn't have to use this trick.  This input system needs more flexibility in its public interface!
+				//bWasHandled = true;
+				//bOutIsInputCaptured = false;
+			}
+		}
+	}
+}
+
 #undef LOCTEXT_NAMESPACE
diff --git a/Engine/Source/Editor/MeshPaint/Public/MeshPaintEdMode.h b/Engine/Source/Editor/MeshPaint/Public/MeshPaintEdMode.h
--- a/Engine/Source/Editor/MeshPaint/Public/MeshPaintEdMode.h
+++ b/Engine/Source/Editor/MeshPaint/Public/MeshPaintEdMode.h
@@ -5,6 +5,8 @@
 #include "GenericOctree.h"
 #include "GenericOctreePublic.h"
 #include "Engine/StaticMesh.h"
+#include "IVREditorModule.h"
+
 
 class IMeshPaintGeometryAdapter;
 
@@ -399,6 +401,9 @@ public:
 	// End of FEdMode interface
 
 
+	/** Called when an editor mode is entered or exited */
+	void OnEditorModeChanged( FEdMode* EditorMode, bool bEntered );
+
 	/** Returns true if we need to force a render/update through based fill/copy */
 	bool IsForceRendered (void) const;
 
@@ -709,11 +714,18 @@ private:
 	/** Called when an asset is about to be reimported */
 	void OnPreReimportAsset(UObject* Object);
 
+	/** Called when the user presses a button on their motion controller device */
+	void OnVRAction( class FEditorViewportClient& ViewportClient, class UViewportInteractor* Interactor,
+		const struct FViewportActionKeyInput& Action, bool& bOutIsInputCaptured, bool& bWasHandled );
+
 private:
 
 	/** Whether we're currently painting */
 	bool bIsPainting;
 
+	/** When painting in VR, this is the hand index that we're painting with.  Otherwise INDEX_NONE. */
+	class UViewportInteractor* PaintingWithInteractorInVR;
+
 	/** Whether we are doing a flood fill the next render */
 	bool bIsFloodFill;
 
diff --git a/Engine/Source/Editor/VREditor/Gizmo/VREditorPlaneTranslationGizmoHandle.cpp b/Engine/Source/Editor/VREditor/Gizmo/VREditorPlaneTranslationGizmoHandle.cpp
--- a/Engine/Source/Editor/VREditor/Gizmo/VREditorPlaneTranslationGizmoHandle.cpp
+++ b/Engine/Source/Editor/VREditor/Gizmo/VREditorPlaneTranslationGizmoHandle.cpp
@@ -2,7 +2,7 @@
 
 #include "VREditorModule.h"
 #include "VREditorPlaneTranslationGizmoHandle.h"
-#include "VREditorBaseTransformGizmo.h"
+#include "VREditorTransformGizmo.h"
 #include "UnitConversion.h"
 
 UVREditorPlaneTranslationGizmoHandleGroup::UVREditorPlaneTranslationGizmoHandleGroup()
@@ -46,14 +46,16 @@ UVREditorPlaneTranslationGizmoHandleGroup::UVREditorPlaneTranslationGizmoHandleG
 	}
 }
 
-void UVREditorPlaneTranslationGizmoHandleGroup::UpdateGizmoHandleGroup( const FTransform& LocalToWorld, const FBox& LocalBounds, const FVector ViewLocation, bool bAllHandlesVisible, class UActorComponent* DraggingHandle, const TArray< UActorComponent* >& HoveringOverHandles, float AnimationAlpha, float GizmoScale, const float GizmoHoverScale, const float GizmoHoverAnimationDuration, bool& bOutIsHoveringOrDraggingThisHandleGroup )
+void UVREditorPlaneTranslationGizmoHandleGroup::UpdateGizmoHandleGroup(const FTransform& LocalToWorld, const FBox& LocalBounds, const FVector ViewLocation, bool bAllHandlesVisible, class UActorComponent* DraggingHandle,
+	const TArray< UActorComponent* >& HoveringOverHandles, float AnimationAlpha, float GizmoScale, const float GizmoHoverScale, const float GizmoHoverAnimationDuration, bool& bOutIsHoveringOrDraggingThisHandleGroup )
 {
 	// Call parent implementation (updates hover animation)
-	Super::UpdateGizmoHandleGroup( LocalToWorld, LocalBounds, ViewLocation, bAllHandlesVisible, DraggingHandle, HoveringOverHandles, AnimationAlpha, GizmoScale, GizmoHoverScale, GizmoHoverAnimationDuration, bOutIsHoveringOrDraggingThisHandleGroup );
+	Super::UpdateGizmoHandleGroup(LocalToWorld, LocalBounds, ViewLocation, bAllHandlesVisible, DraggingHandle, HoveringOverHandles, AnimationAlpha,
+		GizmoScale, GizmoHoverScale, GizmoHoverAnimationDuration, bOutIsHoveringOrDraggingThisHandleGroup );
 
 	for (int32 HandleIndex = 0; HandleIndex < Handles.Num(); ++HandleIndex)
 	{
-		FVREditorGizmoHandle& Handle = Handles[ HandleIndex ];
+		FGizmoHandle& Handle = Handles[ HandleIndex ];
 		UStaticMeshComponent* PlaneTranslationHandle = Handle.HandleMesh;
 		if (PlaneTranslationHandle != nullptr)	// Can be null if no handle for this specific placement
 		{
@@ -116,4 +118,4 @@ ETransformGizmoInteractionType UVREditorPlaneTranslationGizmoHandleGroup::GetInt
 EGizmoHandleTypes UVREditorPlaneTranslationGizmoHandleGroup::GetHandleType() const
 {
 	return EGizmoHandleTypes::Translate;
-}
\ No newline at end of file
+}
diff --git a/Engine/Source/Editor/VREditor/Gizmo/VREditorPlaneTranslationGizmoHandle.h b/Engine/Source/Editor/VREditor/Gizmo/VREditorPlaneTranslationGizmoHandle.h
--- a/Engine/Source/Editor/VREditor/Gizmo/VREditorPlaneTranslationGizmoHandle.h
+++ b/Engine/Source/Editor/VREditor/Gizmo/VREditorPlaneTranslationGizmoHandle.h
@@ -2,14 +2,14 @@
 
 #pragma once
 
-#include "VREditorGizmoHandle.h"
+#include "VIGizmoHandle.h"
 #include "VREditorPlaneTranslationGizmoHandle.generated.h"
 
 /**
  * Gizmo handle for translating on a plane
  */
 UCLASS()
-class UVREditorPlaneTranslationGizmoHandleGroup : public UVREditorGizmoHandleGroup
+class UVREditorPlaneTranslationGizmoHandleGroup : public UGizmoHandleGroup
 {
 	GENERATED_BODY()
 
@@ -19,7 +19,8 @@ public:
 	UVREditorPlaneTranslationGizmoHandleGroup();
 	
 	/** Updates the gizmo handles */
-	virtual void UpdateGizmoHandleGroup( const FTransform& LocalToWorld, const FBox& LocalBounds, const FVector ViewLocation, bool bAllHandlesVisible, class UActorComponent* DraggingHandle, const TArray< UActorComponent* >& HoveringOverHandles, float AnimationAlpha, float GizmoScale, const float GizmoHoverScale, const float GizmoHoverAnimationDuration, bool& bOutIsHoveringOrDraggingThisHandleGroup ) override;
+	void UpdateGizmoHandleGroup(const FTransform& LocalToWorld, const FBox& LocalBounds, const FVector ViewLocation, bool bAllHandlesVisible, class UActorComponent* DraggingHandle,
+		const TArray< UActorComponent* >& HoveringOverHandles, float AnimationAlpha, float GizmoScale, const float GizmoHoverScale, const float GizmoHoverAnimationDuration, bool& bOutIsHoveringOrDraggingThisHandleGroup ) override;
 
 	/** Gets the InteractionType for this Gizmo handle */
 	virtual ETransformGizmoInteractionType GetInteractionType() const override;
diff --git a/Engine/Source/Editor/VREditor/Gizmo/VREditorRotationGizmoHandle.cpp b/Engine/Source/Editor/VREditor/Gizmo/VREditorRotationGizmoHandle.cpp
--- a/Engine/Source/Editor/VREditor/Gizmo/VREditorRotationGizmoHandle.cpp
+++ b/Engine/Source/Editor/VREditor/Gizmo/VREditorRotationGizmoHandle.cpp
@@ -46,15 +46,17 @@ UVREditorRotationGizmoHandleGroup::UVREditorRotationGizmoHandleGroup()
 	}
 }
 
-void UVREditorRotationGizmoHandleGroup::UpdateGizmoHandleGroup( const FTransform& LocalToWorld, const FBox& LocalBounds, const FVector ViewLocation, bool bAllHandlesVisible, class UActorComponent* DraggingHandle, const TArray< UActorComponent* >& HoveringOverHandles, float AnimationAlpha, float GizmoScale, const float GizmoHoverScale, const float GizmoHoverAnimationDuration, bool& bOutIsHoveringOrDraggingThisHandleGroup )
+void UVREditorRotationGizmoHandleGroup::UpdateGizmoHandleGroup( const FTransform& LocalToWorld, const FBox& LocalBounds, const FVector ViewLocation, bool bAllHandlesVisible, class UActorComponent* DraggingHandle, 
+	const TArray< UActorComponent* >& HoveringOverHandles, float AnimationAlpha, float GizmoScale, const float GizmoHoverScale, const float GizmoHoverAnimationDuration, bool& bOutIsHoveringOrDraggingThisHandleGroup )
 {
 	// Call parent implementation (updates hover animation)
-	Super::UpdateGizmoHandleGroup( LocalToWorld, LocalBounds, ViewLocation, bAllHandlesVisible, DraggingHandle, HoveringOverHandles, AnimationAlpha, GizmoScale, GizmoHoverScale, GizmoHoverAnimationDuration, bOutIsHoveringOrDraggingThisHandleGroup );
+	Super::UpdateGizmoHandleGroup(LocalToWorld, LocalBounds, ViewLocation, bAllHandlesVisible, DraggingHandle, HoveringOverHandles,
+		AnimationAlpha, GizmoScale, GizmoHoverScale, GizmoHoverAnimationDuration, bOutIsHoveringOrDraggingThisHandleGroup);
 
 	// Place the rotation handles
 	for (int32 HandleIndex = 0; HandleIndex < Handles.Num(); ++HandleIndex)
 	{
-		FVREditorGizmoHandle& Handle = Handles[ HandleIndex ];
+		FGizmoHandle& Handle = Handles[ HandleIndex ];
 
 		UStaticMeshComponent* RotationHandle = Handle.HandleMesh;
 		if (RotationHandle != nullptr)	// Can be null if no handle for this specific placement
diff --git a/Engine/Source/Editor/VREditor/Gizmo/VREditorTranslationGizmoHandle.cpp b/Engine/Source/Editor/VREditor/Gizmo/VREditorTranslationGizmoHandle.cpp
--- a/Engine/Source/Editor/VREditor/Gizmo/VREditorTranslationGizmoHandle.cpp
+++ b/Engine/Source/Editor/VREditor/Gizmo/VREditorTranslationGizmoHandle.cpp
@@ -46,15 +46,17 @@ UVREditorTranslationGizmoHandleGroup::UVREditorTranslationGizmoHandleGroup()
 	}
 }
 
-void UVREditorTranslationGizmoHandleGroup::UpdateGizmoHandleGroup( const FTransform& LocalToWorld, const FBox& LocalBounds, const FVector ViewLocation, bool bAllHandlesVisible, class UActorComponent* DraggingHandle, const TArray< UActorComponent* >& HoveringOverHandles, float AnimationAlpha, float GizmoScale, const float GizmoHoverScale, const float GizmoHoverAnimationDuration, bool& bOutIsHoveringOrDraggingThisHandleGroup )
+void UVREditorTranslationGizmoHandleGroup::UpdateGizmoHandleGroup( const FTransform& LocalToWorld, const FBox& LocalBounds, const FVector ViewLocation, bool bAllHandlesVisible, class UActorComponent* DraggingHandle, 
+	const TArray< UActorComponent* >& HoveringOverHandles, float AnimationAlpha, float GizmoScale, const float GizmoHoverScale, const float GizmoHoverAnimationDuration, bool& bOutIsHoveringOrDraggingThisHandleGroup )
 {
 	// Call parent implementation (updates hover animation)
-	Super::UpdateGizmoHandleGroup( LocalToWorld, LocalBounds, ViewLocation, bAllHandlesVisible, DraggingHandle, HoveringOverHandles, AnimationAlpha, GizmoScale, GizmoHoverScale, GizmoHoverAnimationDuration, bOutIsHoveringOrDraggingThisHandleGroup );
+	Super::UpdateGizmoHandleGroup(LocalToWorld, LocalBounds, ViewLocation, bAllHandlesVisible, DraggingHandle, HoveringOverHandles,
+		AnimationAlpha, GizmoScale, GizmoHoverScale, GizmoHoverAnimationDuration, bOutIsHoveringOrDraggingThisHandleGroup );
 
 	// Place the translation handles
 	for (int32 HandleIndex = 0; HandleIndex < Handles.Num(); ++HandleIndex)
 	{
-		FVREditorGizmoHandle& Handle = Handles[ HandleIndex ];
+		FGizmoHandle& Handle = Handles[ HandleIndex ];
 
 		UStaticMeshComponent* TranslationHandle = Handle.HandleMesh;
 		if (TranslationHandle != nullptr)	// Can be null if no handle for this specific placement
diff --git a/Engine/Source/Editor/VREditor/Interactables/VREditorButton.cpp b/Engine/Source/Editor/VREditor/Interactables/VREditorButton.cpp
new file mode 100644
--- /dev/null
+++ b/Engine/Source/Editor/VREditor/Interactables/VREditorButton.cpp
@@ -0,0 +1,61 @@
+// Copyright 1998-2015 Epic Games, Inc. All Rights Reserved.
+
+#include "VREditorModule.h"
+#include "VREditorButton.h"
+#include "ViewportInteractor.h"
+
+AVREditorButton::AVREditorButton()
+{
+	const bool bAllowGizmoLighting = false;
+
+	// Setup uniform scaling
+	UStaticMesh* ButtonMesh = nullptr;
+	{
+		static ConstructorHelpers::FObjectFinder<UStaticMesh> ObjectFinder( TEXT( "/Engine/VREditor/TransformGizmo/UniformScaleHandle" ) );
+		ButtonMesh = ObjectFinder.Object;
+		check( ButtonMesh != nullptr );
+	}
+
+	ButtonMeshComponent = CreateDefaultSubobject<UStaticMeshComponent>( TEXT( "ButtonMeshComponent" ) );
+	check( ButtonMeshComponent != nullptr );
+
+	ButtonMeshComponent->SetStaticMesh( ButtonMesh );
+	ButtonMeshComponent->SetMobility( EComponentMobility::Movable );
+
+	ButtonMeshComponent->SetCollisionEnabled( ECollisionEnabled::QueryOnly );
+	ButtonMeshComponent->SetCollisionResponseToAllChannels( ECR_Ignore );
+	ButtonMeshComponent->SetCollisionResponseToChannel( COLLISION_GIZMO, ECollisionResponse::ECR_Block );
+	ButtonMeshComponent->SetCollisionObjectType( COLLISION_GIZMO );
+
+	ButtonMeshComponent->bGenerateOverlapEvents = false;
+	ButtonMeshComponent->SetCanEverAffectNavigation( false );
+	ButtonMeshComponent->bCastDynamicShadow = bAllowGizmoLighting;
+	ButtonMeshComponent->bCastStaticShadow = false;
+	ButtonMeshComponent->bAffectDistanceFieldLighting = bAllowGizmoLighting;
+	ButtonMeshComponent->bAffectDynamicIndirectLighting = bAllowGizmoLighting;
+}
+
+void AVREditorButton::OnPressed( UViewportInteractor* Interactor, const FHitResult& InHitResult, bool& bOutResultedInDrag  )
+{
+	OnPressedEvent.Broadcast( this, Interactor );
+}
+
+void AVREditorButton::OnHover( UViewportInteractor* Interactor )
+{
+	OnHoverEvent.Broadcast( this, Interactor );
+}
+
+void AVREditorButton::OnHoverEnter( UViewportInteractor* Interactor, const FHitResult& InHitResult  )
+{
+	OnHoverEnterEvent.Broadcast( this, Interactor );
+}
+
+void AVREditorButton::OnHoverLeave( UViewportInteractor* Interactor, const UActorComponent* NewComponent )
+{
+	OnHoverLeaveEvent.Broadcast( this, Interactor );
+}
+
+void AVREditorButton::OnDragRelease( UViewportInteractor* Interactor )
+{
+	OnReleaseEvent.Broadcast( this, Interactor );
+}
\ No newline at end of file
diff --git a/Engine/Source/Editor/VREditor/Interactables/VREditorButton.h b/Engine/Source/Editor/VREditor/Interactables/VREditorButton.h
new file mode 100644
--- /dev/null
+++ b/Engine/Source/Editor/VREditor/Interactables/VREditorButton.h
@@ -0,0 +1,75 @@
+// Copyright 1998-2015 Epic Games, Inc. All Rights Reserved.
+
+#pragma once
+
+#include "VREditorBaseActor.h"
+#include "ViewportInteractableInterface.h"
+#include "VREditorButton.generated.h"
+
+//Forward declarations
+class UViewportInteractor;
+
+
+/**
+* A button for VR Editor
+*/
+UCLASS()
+class AVREditorButton : public AVREditorBaseActor, public IViewportInteractableInterface
+{
+	GENERATED_BODY()
+
+public:
+
+	/** Default constructor which sets up safe defaults */
+	AVREditorButton();
+
+	// Begin IViewportInteractableInterface
+	virtual void OnPressed( UViewportInteractor* Interactor, const FHitResult& InHitResult, bool& bOutResultedInDrag  ) override;
+	virtual void OnHover( UViewportInteractor* Interactor ) override;
+	virtual void OnHoverEnter( UViewportInteractor* Interactor, const FHitResult& InHitResult ) override;
+	virtual void OnHoverLeave( UViewportInteractor* Interactor, const UActorComponent* NewComponent ) override;
+	virtual void OnDragRelease( UViewportInteractor* Interactor ) override;
+	virtual class UViewportDragOperationComponent* GetDragOperationComponent() override { return nullptr; };
+	// End IViewportInteractableInterface
+
+	/** Gets the on pressed event */
+	DECLARE_EVENT_TwoParams( AVREditorButton, FOnPressedEvent, AVREditorButton*, UViewportInteractor* );
+	FOnPressedEvent& GetOnPressed() { return OnPressedEvent; }
+
+	/** Gets the on hover event */
+	DECLARE_EVENT_TwoParams( AVREditorButton, FOnHoverEvent, AVREditorButton*, UViewportInteractor* );
+	FOnHoverEvent& GetOnHover() { return OnHoverEvent; }
+
+	/** Gets the on hover enter event */
+	DECLARE_EVENT_TwoParams( AVREditorButton, FOnHoverEnterEvent, AVREditorButton*, UViewportInteractor* );
+	FOnHoverEnterEvent& GetOnHoverEnter() { return OnHoverEnterEvent; }
+
+	/** Gets the on hover leave event */
+	DECLARE_EVENT_TwoParams( AVREditorButton, FOnHoverLeaveEvent, AVREditorButton*, UViewportInteractor* );
+	FOnHoverLeaveEvent	& GetOnHoverLeave() { return OnHoverLeaveEvent; }
+
+	/** Gets the on release event */
+	DECLARE_EVENT_TwoParams( AVREditorButton, FOnReleaseEvent, AVREditorButton*, UViewportInteractor* );
+	FOnReleaseEvent& GetOnRelease() { return OnReleaseEvent; }
+
+private:
+
+	/** Mesh for the button to interact with */
+	UPROPERTY()
+	UStaticMeshComponent* ButtonMeshComponent;
+
+	/** Broadcasted when pressed on this interactable by an interactor */
+	FOnPressedEvent OnPressedEvent;
+
+	/** Broadcasted when interactor is hovering over this interactable */
+	FOnHoverEvent OnHoverEvent;
+
+	/** Broadcasted when interactor started hovering over this interactable */
+	FOnHoverEnterEvent OnHoverEnterEvent;
+
+	/** Broadcasted when interactor stopped hovering over this interactable */
+	FOnHoverLeaveEvent OnHoverLeaveEvent;
+
+	/** Broadcasted when interactor released this interactable */
+	FOnReleaseEvent OnReleaseEvent;
+};
\ No newline at end of file
diff --git a/Engine/Source/Editor/VREditor/Public/IVREditorMode.h b/Engine/Source/Editor/VREditor/Public/IVREditorMode.h
--- a/Engine/Source/Editor/VREditor/Public/IVREditorMode.h
+++ b/Engine/Source/Editor/VREditor/Public/IVREditorMode.h
@@ -2,57 +2,18 @@
 
 #pragma once
 
-
-namespace VREditorConstants
-{
-	/** The number of virtual hands we support.  Always two. */
-	const int32 NumVirtualHands = 2;
-
-	/** Index of the left hand */
-	const int32 LeftHandIndex = 0;
-
-	/** Index of the right hand */
-	const int32 RightHandIndex = 1;
-};
-
-
 /**
  * Types of actions that can be performed with VR controller devices
  */
-enum class EVRActionType
+namespace VRActionTypes
 {
-	NoAction,
-	Undo,
-	Redo,
-	Touch,
-	Modifier,
-	WorldMovement,
-	SelectAndMove,
-	SelectAndMove_LightlyPressed,
-	Delete,
-	ConfirmRadialSelection,
-	TrackpadPositionX,
-	TrackpadPositionY,
-	TriggerAxis,
-	// ...
-	TotalActionTypes
-};
-
-
-/**
- * Represents a generic action bound to a specific VR "hand"
- */
-struct FVRAction
-{
-	FVRAction( EVRActionType InActionType, int32 InHandIndex )
-		: ActionType( InActionType )
-		, HandIndex( InHandIndex )
-	{}
-
-	EVRActionType ActionType;
-	int32 HandIndex;
-};
-
+	static const FName Touch( "Touch" );
+	static const FName Modifier( "Modifier" );
+	static const FName ConfirmRadialSelection( "ConfirmRadialSelection" );
+	static const FName TrackpadPositionX( "TrackpadPositionX" );
+	static const FName TrackpadPositionY( "TrackpadPositionY" );
+	static const FName TriggerAxis( "TriggerAxis" );
+}
 
 /**
  * VR Editor Mode public interface
@@ -63,17 +24,11 @@ class IVREditorMode : public FEdMode
 public:
 
 	/**
-	 * Gets the start and end point of the laser pointer for the specified hand
-	 *
-	 * @param HandIndex				Index of the hand to use
-	 * @param LasertPointerStart	(Out) The start location of the laser pointer in world space
-	 * @param LasertPointerEnd		(Out) The end location of the laser pointer in world space
-	 * @param bEvenIfUIIsInFront	If true, returns a laser pointer even if the hand has UI in front of it (defaults to false)
-	 * @param LaserLengthOverride	If zero the default laser length (VREdMode::GetLaserLength) is used
+	 * Gets the world space transform of the HMD (head)
 	 *
-	 * @return	True if we have motion controller data for this hand and could return a valid result
+	 * @return	World space space HMD transform
 	 */
-	virtual bool GetLaserPointer( const int32 HandIndex, FVector& LaserPointerStart, FVector& LaserPointerEnd, const bool bEvenIfUIIsInFront = false, const float LaserLengthOverride = 0.0f ) const = 0;
+	virtual FTransform GetHeadTransform() const = 0;
 
 	/**
 	 * Gets our avatar's mesh actor
@@ -82,26 +37,31 @@ public:
 	 */
 	virtual class AActor* GetAvatarMeshActor() = 0;
 
-	/**
-	 * Checks to see if input is captured for the specified action
-	 *
-	 * @param	VRAction	The VR action that might be captured (the action type, and hand)
-	 *
-	 * @return	True if input is currently captured for this action
-	 */
-	virtual bool IsInputCaptured( const FVRAction VRAction ) const = 0;
-
-	/** Event that's triggered when the user presses a button on their motion controller device.  Check bIsInputCaptured first, 
-	    to see whether you should actually handle input.  Set bIsInputCaptured to true if you want to capture input, or clear it
-	    when finished capturing, if you were the captor.  Set bWasHandled to true if you reacted to the event.  Remember to check
-		bWasHandled yourself first, as there may be sibling subscribers that handle the event before you! */
-	DECLARE_EVENT_FiveParams( IVREditorMode, FOnVRAction, class FEditorViewportClient& /* ViewportClient */, const FVRAction /* VRAction */, const EInputEvent /* Event */, bool& /* bIsInputCaptured */, bool& /* bWasHandled */ );
-	virtual FOnVRAction& OnVRAction() = 0;
-
-	/** Event that fires every frame to update hover state based on what's under the cursor.  Set bWasHandled to true if you detect 
-	    something to hover, and set HoverImpactPoint to the world space location you hit.  Remember to check bWasHandled yourself 
-		first, as there may be sibling subscribers that handle the event before you! */
-	DECLARE_EVENT_FiveParams( IVREditorMode, FOnVRHoverUpdate, class FEditorViewportClient& /* ViewportClient */, const int32 /* HandIndex */, FVector& /* HoverImpactPoint */, bool& /* bIsHoveringOverUI */, bool& /* bWasHandled */ );
-	virtual FOnVRHoverUpdate& OnVRHoverUpdate() = 0;
+	/** Gets the world currently used by the worldinteraction of this mode */
+	virtual UWorld* GetWorld() const = 0;
+
+	/** Gets the world space transform of the calibrated VR room origin.  When using a seated VR device, this will feel like the
+	camera's world transform (before any HMD positional or rotation adjustments are applied.) */
+	virtual FTransform GetRoomTransform() const = 0;
+
+	/** Sets a new transform for the room, in world space.  This is basically setting the editor's camera transform for the viewport */
+	virtual void SetRoomTransform( const FTransform& NewRoomTransform ) = 0;
+
+	/** Gets the transform of the user's HMD in room space */
+	virtual FTransform GetRoomSpaceHeadTransform() const = 0;
+
+	/** Gets access to the world interaction system (const) */
+	virtual const class UViewportWorldInteraction& GetWorldInteraction() const = 0;
+
+	/** Gets access to the world interaction system */
+	virtual class UViewportWorldInteraction& GetWorldInteraction() = 0;
+
+	/** If the mode was completely initialized */
+	virtual bool IsFullyInitialized() const = 0;
+
+	/** Gets the event for update which is broadcasted by the mode itself after getting latest input states */
+	DECLARE_EVENT_OneParam( IVREditorMode, FOnVRTickHandle, const float /* DeltaTime */ );
+	virtual FOnVRTickHandle& OnTickHandle() = 0;
+
 };
 
diff --git a/Engine/Source/Editor/VREditor/Public/VREditorInteractor.h b/Engine/Source/Editor/VREditor/Public/VREditorInteractor.h
new file mode 100644
--- /dev/null
+++ b/Engine/Source/Editor/VREditor/Public/VREditorInteractor.h
@@ -0,0 +1,164 @@
+// Copyright 1998-2015 Epic Games, Inc. All Rights Reserved.
+
+#pragma once
+
+#include "Editor/ViewportInteraction/Public/ViewportInteractor.h"
+#include "VREditorInteractor.generated.h"
+
+/**
+ * VREditor default interactor
+ */
+UCLASS()
+class VREDITOR_API UVREditorInteractor : public UViewportInteractor
+{
+	GENERATED_UCLASS_BODY()
+
+public:
+
+	virtual ~UVREditorInteractor();
+
+	/** Initialize default values */
+	virtual void Init( class FVREditorMode* InVRMode );
+
+	/** Gets the owner of this system */
+	class FVREditorMode& GetVRMode()
+	{
+		return *VRMode;
+	}
+
+	/** Gets the owner of this system (const) */
+	const class FVREditorMode& GetVRMode() const
+	{
+		return *VRMode;
+	}
+
+	// ViewportInteractorInterface overrides
+	virtual void Shutdown() override;
+	virtual void Tick( const float DeltaTime ) override;
+	virtual FHitResult GetHitResultFromLaserPointer( TArray<AActor*>* OptionalListOfIgnoredActors = nullptr, const bool bIgnoreGizmos = false,
+		TArray<UClass*>* ObjectsInFrontOfGizmo = nullptr, const bool bEvenIfBlocked = false, const float LaserLengthOverride = 0.0f ) override;
+	virtual void ResetHoverState( const float DeltaTime ) override;
+	virtual void OnStartDragging( UActorComponent* ClickedComponent, const FVector& HitLocation, const bool bIsPlacingActors ) override;
+
+	/** Returns the slide delta for pushing and pulling objects. Needs to be implemented by derived classes (e.g. touchpad for vive controller or scrollweel for mouse ) */
+	virtual float GetSlideDelta();
+
+	//
+	// Getters and setters
+	//
+
+	/** Gets if this interactor is hovering over UI */
+	bool IsHoveringOverUI() const;
+	
+	/** Sets if there is a docked window on this interactor */
+	void SetHasUIInFront( const bool bInHasUIInFront );
+	
+	/** Check if there is any docked window on this interactor */
+	bool HasUIInFront() const;
+
+	/** Sets if the quick menu is on this interactor */
+	void SetHasUIOnForearm( const bool bInHasUIOnForearm );
+	
+	/** Check if the quick menu is on this interactor */
+	bool HasUIOnForearm() const;
+
+	/** Gets the current hovered widget component if any */
+	UWidgetComponent* GetHoveringOverWidgetComponent() const;
+
+	/** Sets the current hovered widget component */
+	void SetHoveringOverWidgetComponent( UWidgetComponent* NewHoveringOverWidgetComponent );
+
+	/** If the modifier key is currently pressed */
+	bool IsModifierPressed() const;
+
+	/** Sets if the interactor is clicking on any UI */
+	void SetIsClickingOnUI( const bool bInIsClickingOnUI );
+	
+	/** Gets if the interactor is clicking on any UI */
+	bool IsClickingOnUI() const;
+
+	/** Sets if the interactor is hovering over any UI */
+	void SetIsHoveringOverUI( const bool bInIsHoveringOverUI );
+
+	/** Sets if the interactor is right  hover over any UI */
+	void SetIsRightClickingOnUI( const bool bInIsRightClickingOnUI );
+
+	/** Gets if the interactor is right clicking on UI */
+	bool IsRightClickingOnUI() const;
+
+	/** Gets if the interactor is right clicking on UI */
+	void SetLastClickReleaseTime( const double InLastClickReleaseTime );
+
+	/** Gets last time the interactor released */
+	double GetLastClickReleaseTime() const;
+
+	/** Sets the UI scroll velocity */
+	void SetUIScrollVelocity( const float InUIScrollVelocity );
+
+	/** Gets the UI scroll velocity */
+	float GetUIScrollVelocity() const;
+
+	/** Gets the trigger value */
+	float GetSelectAndMoveTriggerValue() const;
+	
+	/* ViewportInteractor overrides, checks if the laser is blocked by UI */
+	virtual bool GetIsLaserBlocked() override;
+	
+protected:
+
+	/** The mode that owns this interactor */
+	class FVREditorMode* VRMode;
+
+	//
+	// General input @todo: VREditor: Should this be general (non-UI) in interactordata ?
+	//
+
+	/** Is the Modifier button held down? */
+	bool bIsModifierPressed;
+
+	/** Current trigger pressed amount for 'select and move' (0.0 - 1.0) */
+	float SelectAndMoveTriggerValue;
+
+private:
+
+	//
+	// UI
+	//
+
+	/** True if a floating UI panel is attached to the front of the hand, and we shouldn't bother drawing a laser
+	pointer or enabling certain other features */
+	bool bHasUIInFront;
+
+	/** True if a floating UI panel is attached to our forearm, so we shouldn't bother drawing help labels */
+	bool bHasUIOnForearm;
+
+	/** True if we're currently holding the 'SelectAndMove' button down after clicking on UI */
+	bool bIsClickingOnUI;
+
+	/** When bIsClickingOnUI is true, this will be true if we're "right clicking".  That is, the Modifier key was held down at the time that the user clicked */
+	bool bIsRightClickingOnUI;
+
+	 /** True if we're hovering over UI right now.  When hovering over UI, we don't bother drawing a see-thru laser pointer */
+	bool bIsHoveringOverUI;
+	
+	/** Inertial scrolling -- how fast to scroll the mousewheel over UI */
+	float UIScrollVelocity;	
+
+	/** Last real time that we released the 'SelectAndMove' button on UI.  This is used to detect double-clicks. */
+	double LastClickReleaseTime;
+
+protected:
+
+	//
+	// Help
+	//
+
+	/** True if we want help labels to be visible right now, otherwise false */
+	bool bWantHelpLabels;
+
+	/** Help labels for buttons on the motion controllers */
+	TMap< FKey, class AFloatingText* > HelpLabels;
+
+	/** Time that we either started showing or hiding help labels (for fade transitions) */
+	FTimespan HelpLabelShowOrHideStartTime;
+};
\ No newline at end of file
diff --git a/Engine/Source/Editor/VREditor/Teleporter/VREditorTeleporter.cpp b/Engine/Source/Editor/VREditor/Teleporter/VREditorTeleporter.cpp
new file mode 100644
--- /dev/null
+++ b/Engine/Source/Editor/VREditor/Teleporter/VREditorTeleporter.cpp
@@ -0,0 +1,144 @@
+// Copyright 1998-2015 Epic Games, Inc. All Rights Reserved.
+
+#include "VREditorModule.h"
+#include "VREditorTeleporter.h"
+#include "IVREditorMode.h"
+#include "ViewportWorldInteraction.h"
+#include "VREditorInteractor.h"
+
+#include "Kismet/GameplayStatics.h"
+
+namespace VREd
+{
+	static FAutoConsoleVariable TeleportLerpTime( TEXT( "VREd.TeleportLerpTime" ), 0.1f, TEXT( "The lerp time to teleport" ) );
+	static FAutoConsoleVariable TeleportOffset( TEXT( "VREd.TeleportOffset" ), 100.0f, TEXT( "The offset from the hitresult towards the controller" ) );
+	static FAutoConsoleVariable TeleportLaserPointerLength( TEXT( "VREd.TeleportLaserPointerLength" ), 500000.0f, TEXT( "Distance of the LaserPointer for teleporting" ) );
+}
+
+UVREditorTeleporter::UVREditorTeleporter( const FObjectInitializer& Initializer ) :
+	Super( Initializer ),
+	VRMode( nullptr ),
+	bIsTeleporting( false ),
+	TeleportLerpAlpha( 0 ),
+	TeleportStartLocation( FVector::ZeroVector ),
+	TeleportGoalLocation( FVector::ZeroVector ),
+	bJustTeleported( false )
+{
+	// Load sounds
+	TeleportSound = LoadObject<USoundCue>( nullptr, TEXT( "/Engine/VREditor/Sounds/VR_teleport_Cue" ) );
+	check( TeleportSound != nullptr );
+}
+
+void UVREditorTeleporter::Init( IVREditorMode* InMode )
+{
+	VRMode = InMode;
+	VRMode->OnTickHandle().AddUObject( this, &UVREditorTeleporter::Tick );
+	VRMode->GetWorldInteraction().OnViewportInteractionInputAction().AddUObject( this, &UVREditorTeleporter::OnInteractorAction );
+}
+
+void UVREditorTeleporter::Shutdown()
+{
+	VRMode->OnTickHandle().RemoveAll( this );
+	VRMode->GetWorldInteraction().OnViewportInteractionInputAction().RemoveAll( this );
+	TeleportSound = nullptr;
+}
+
+void UVREditorTeleporter::Tick( const float DeltaTime )
+{
+	if ( bIsTeleporting )
+	{
+		Teleport( DeltaTime );
+
+		//if ( bJustTeleported && !bIsTeleporting )
+		//{
+		//	bJustTeleported = false;
+		//}
+	}
+}
+
+void UVREditorTeleporter::StartTeleport( UViewportInteractor* Interactor )
+{
+	UVREditorInteractor* VREditorInteractor = Cast<UVREditorInteractor>( Interactor );
+	if ( VREditorInteractor && !VREditorInteractor->IsHoveringOverUI() )
+	{
+		FVector LaserPointerStart, LaserPointerEnd;
+		if ( Interactor->GetLaserPointer( LaserPointerStart, LaserPointerEnd ) )
+		{
+			FHitResult HitResult = Interactor->GetHitResultFromLaserPointer( nullptr, true, nullptr, false, VREd::TeleportLaserPointerLength->GetFloat() );
+			if ( HitResult.bBlockingHit )
+			{
+				const FVector HeadLocation = VRMode->GetHeadTransform().GetLocation();
+				const FVector RoomLocation = VRMode->GetRoomTransform().GetLocation();
+
+				const FVector Offset = ( LaserPointerStart - HitResult.Location ).GetSafeNormal() *
+					VREd::TeleportOffset->GetFloat() * VRMode->GetWorldInteraction().GetWorldScaleFactor(); 	// Offset from the hitresult, TeleportOffset is in centimeter
+				const FVector NewHandLocation = HitResult.Location + Offset;								// New world location offset from the hitresult
+				const FVector HandHeadOffset = HeadLocation - Interactor->GetTransform().GetLocation();		// The offset between the hand and the head
+				const FVector NewHeadLocation = NewHandLocation + HandHeadOffset;							// The location in world space where the head is going to be
+																											// The new roomspace location in world space
+				const FVector HeadRoomOffset = HeadLocation - RoomLocation;									// Where the head is going to be in world space
+				const FVector NewRoomSpaceLocation = NewHeadLocation - HeadRoomOffset;						// Result, where the roomspace is going to be in world space
+				
+																											// Set values to start teleporting
+				bIsTeleporting = true;
+				TeleportLerpAlpha = 0.0f;
+				TeleportStartLocation = RoomLocation;
+				TeleportGoalLocation = NewRoomSpaceLocation;
+
+				UGameplayStatics::PlaySound2D( VRMode->GetWorld(), TeleportSound );
+			}
+		}
+	}
+}
+
+void UVREditorTeleporter::OnInteractorAction( class FEditorViewportClient& ViewportClient, UViewportInteractor* Interactor, const struct FViewportActionKeyInput& Action, bool& bOutIsInputCaptured, bool& bWasHandled )
+{
+	if ( ( Action.ActionType == ViewportWorldActionTypes::SelectAndMove ||
+		Action.ActionType == ViewportWorldActionTypes::SelectAndMove_LightlyPressed ) )
+	{
+		if ( Action.Event == IE_Pressed )
+		{
+			if ( Action.ActionType == ViewportWorldActionTypes::SelectAndMove_LightlyPressed )
+			{
+				// Check to teleport before checking on actor selection, otherwise objects will be selected when teleporting
+				if	( Interactor->GetDraggingMode() == EViewportInteractionDraggingMode::World &&
+					( Interactor->GetOtherInteractor() == nullptr || Interactor->GetOtherInteractor()->GetDraggingMode() != EViewportInteractionDraggingMode::AssistingDrag ) )
+				{
+					StartTeleport( Interactor );
+					bJustTeleported = true;
+					bWasHandled = true;
+				}
+			}
+			else if ( Action.ActionType == ViewportWorldActionTypes::SelectAndMove && bJustTeleported )
+			{
+				//Block the input
+				bJustTeleported = false;
+				bWasHandled = true;
+			}
+
+			// Nothing should happen when the hand has a UI in front, otherwise it will deselect everything
+			if ( Interactor->GetIsLaserBlocked() )
+			{
+				bWasHandled = true;
+			}
+		}
+	}
+}
+
+void UVREditorTeleporter::Teleport( const float DeltaTime )
+{
+	FTransform RoomTransform = VRMode->GetRoomTransform();
+	TeleportLerpAlpha += DeltaTime;
+
+	if ( TeleportLerpAlpha > VREd::TeleportLerpTime->GetFloat() )
+	{
+		// Teleporting is finished
+		TeleportLerpAlpha = VREd::TeleportLerpTime->GetFloat();
+		bIsTeleporting = false;
+	}
+
+	// Calculate the new position of the roomspace
+	FVector NewLocation = FMath::Lerp( TeleportStartLocation, TeleportGoalLocation, TeleportLerpAlpha / VREd::TeleportLerpTime->GetFloat() );
+	RoomTransform.SetLocation( NewLocation );
+	VRMode->SetRoomTransform( RoomTransform );
+}
\ No newline at end of file
diff --git a/Engine/Source/Editor/VREditor/Teleporter/VREditorTeleporter.h b/Engine/Source/Editor/VREditor/Teleporter/VREditorTeleporter.h
new file mode 100644
--- /dev/null
+++ b/Engine/Source/Editor/VREditor/Teleporter/VREditorTeleporter.h
@@ -0,0 +1,70 @@
+// Copyright 1998-2015 Epic Games, Inc. All Rights Reserved.
+
+#pragma once
+
+#include "VREditorTeleporter.generated.h"
+
+/**
+* VR Editor teleport manager
+*/
+UCLASS()
+class UVREditorTeleporter : public UObject
+{
+	GENERATED_UCLASS_BODY()
+
+public:
+
+	/** Initializes the teleport system */
+	void Init( class IVREditorMode* InMode );
+
+	/** Shuts down the teleport system */
+	void Shutdown();
+
+	/** Ticks the current teleport */
+	void Tick( const float DeltaTime );
+
+	/** Returns true if we're actively teleporting */
+	bool IsTeleporting() const
+	{
+		return bIsTeleporting;
+	}
+
+	/** Start teleporting, does a ray trace with the hand passed and calculates the locations for lerp movement in Teleport */
+	void StartTeleport( class UViewportInteractor* Interactor );
+
+private:
+
+	/** Called when the user presses a button on their motion controller device */
+	void OnInteractorAction( class FEditorViewportClient& ViewportClient, UViewportInteractor* Interactor,
+		const struct FViewportActionKeyInput& Action, bool& bOutIsInputCaptured, bool& bWasHandled );
+
+	/** Actually teleport to a position */
+	void Teleport( const float DeltaTime );
+
+	//
+	// Fields
+	//
+
+	/** Owning mode */
+	IVREditorMode* VRMode;
+
+	/** If we are currently teleporting */
+	bool bIsTeleporting;
+
+	/** The current lerp of the teleport between the TeleportStartLocation and the TeleportGoalLocation */
+	float TeleportLerpAlpha;
+
+	/** Set on the current Roomspace location in the world in StartTeleport before doing the actual teleporting */
+	FVector TeleportStartLocation;
+
+	/** The calculated goal location in StartTeleport to move the Roomspace to */
+	FVector TeleportGoalLocation;
+
+	/** So we can teleport on a light press and make sure a second teleport by a fullpress doesn't happen */
+	bool bJustTeleported;
+
+	/** Teleport sound */
+	UPROPERTY()
+	class USoundCue* TeleportSound;
+
+};
\ No newline at end of file
diff --git a/Engine/Source/Editor/VREditor/VREditor.Build.cs b/Engine/Source/Editor/VREditor/VREditor.Build.cs
--- a/Engine/Source/Editor/VREditor/VREditor.Build.cs
+++ b/Engine/Source/Editor/VREditor/VREditor.Build.cs
@@ -19,7 +19,8 @@ namespace UnrealBuildTool.Rules
 					"UMG",
 					"LevelEditor",
 					"HeadMountedDisplay",
-					"Analytics",
+                    "ViewportInteraction",
+					"Analytics"
 				}
 			);
 
diff --git a/Engine/Source/Editor/VREditor/VREditorDockableWindow.h b/Engine/Source/Editor/VREditor/VREditorDockableWindow.h
--- a/Engine/Source/Editor/VREditor/VREditorDockableWindow.h
+++ b/Engine/Source/Editor/VREditor/VREditorDockableWindow.h
@@ -3,13 +3,14 @@
 #pragma once
 
 #include "VREditorFloatingUI.h"
+#include "ViewportInteractableInterface.h"
 #include "VREditorDockableWindow.generated.h"
 
 /**
  * An interactive floating UI panel that can be dragged around
  */
 UCLASS()
-class AVREditorDockableWindow : public AVREditorFloatingUI
+class AVREditorDockableWindow : public AVREditorFloatingUI, public IViewportInteractableInterface
 {
 	GENERATED_BODY()
 	
@@ -17,36 +18,44 @@ public:
 
 	/** Default constructor */
 	AVREditorDockableWindow();
+
+	/** Destructor */
+	~AVREditorDockableWindow();
 	
 	/** Updates the meshes for the UI */
 	virtual void TickManually( float DeltaTime ) override;
 
 	/** Updates the last dragged relative position */
 	void UpdateRelativeRoomTransform();
-
-	/** Enter hover with laser changes the color of SelectionMesh and CloseButtonMesh */
-	void OnEnterHover( const FHitResult& HitResult, const int32 HandIndex );
-
-	/** Leaving hover with laser changes the color of SelectionMesh and CloseButtonMesh */
-	void OnLeaveHover( const int32 HandIndex, const class UActorComponent* NewHoveredComponent );
 	
 	/** Gets the close button component */
-	UStaticMeshComponent* GetCloseButtonMeshComponent()
-	{
-		return CloseButtonMeshComponent;
-	}
+	UStaticMeshComponent* GetCloseButtonMeshComponent() const;
 
 	/** Gets the selection bar component */
-	UStaticMeshComponent* GetSelectionBarMeshComponent()
-	{
-		return SelectionBarMeshComponent;
-	}
+	UStaticMeshComponent* GetSelectionBarMeshComponent() const;
 
-private:
+	/** Gets the distance between the interactor and the window when sta	rting drag */
+	float GetDockSelectDistance() const;
+
+	// IViewportInteractableInterface
+	virtual void OnPressed( UViewportInteractor* Interactor, const FHitResult& InHitResult, bool& bOutResultedInDrag ) override;
+	virtual void OnHover( UViewportInteractor* Interactor ) override;
+	
+	/** Enter hover with laser changes the color of SelectionMesh and CloseButtonMesh */
+	virtual void OnHoverEnter( UViewportInteractor* Interactor, const FHitResult& InHitResult ) override;
+	
+	/** Leaving hover with laser changes the color of SelectionMesh and CloseButtonMesh */
+	virtual void OnHoverLeave( UViewportInteractor* Interactor, const UActorComponent* NewComponent ) override;
+	virtual void OnDragRelease( UViewportInteractor* Interactor ) override;
+	virtual class UViewportDragOperationComponent* GetDragOperationComponent() override;
+
+protected:
 
 	// AVREditorFloatingUI overrides
 	virtual void SetupWidgetComponent() override;
 
+private:
+
 	/** Set the color on the dynamic materials of the selection bar */
 	void SetSelectionBarColor( const FLinearColor& LinearColor );
 
@@ -81,6 +90,8 @@ private:
 	UPROPERTY()
 	class UMaterialInstanceDynamic* CloseButtonTranslucentMID;
 
+	UPROPERTY()
+	class UViewportDragOperationComponent* DragOperationComponent;
 
 	/** True if at least one hand's laser is aiming toward the UI */
 	bool bIsLaserAimingTowardUI;
@@ -88,17 +99,32 @@ private:
 	/** Scalar that ramps up toward 1.0 after the user aims toward the UI */
 	float AimingAtMeFadeAlpha;
 
-
 	/** True if we're hovering over the selection bar */
 	bool bIsHoveringOverSelectionBar;
 
 	/** Scalar that will advance toward 1.0 over time as we hover over the selection bar */
 	float SelectionBarHoverAlpha;
 
-
 	/** True if we're hovering over the close button */
 	bool bIsHoveringOverCloseButton;
 
 	/** Scalar that will advance toward 1.0 over time as we hover over the close button */
 	float CloseButtonHoverAlpha;
+
+	/** Distance from interactor laser to the handle when starting dragging */
+	float DockSelectDistance;
+};
+
+/**
+ * Calculation for dragging a dockable window
+ */
+UCLASS()
+class UDockableWindowDragOperation : public UViewportDragOperation
+{
+	GENERATED_BODY()
+
+public:
+
+	// IViewportDragOperation
+	virtual void ExecuteDrag( class UViewportInteractor* Interactor, IViewportInteractableInterface* Interactable ) override;
 };
\ No newline at end of file
diff --git a/Engine/Source/Editor/VREditor/VREditorFloatingUI.cpp b/Engine/Source/Editor/VREditor/VREditorFloatingUI.cpp
--- a/Engine/Source/Editor/VREditor/VREditorFloatingUI.cpp
+++ b/Engine/Source/Editor/VREditor/VREditorFloatingUI.cpp
@@ -7,6 +7,7 @@
 #include "VREditorMode.h"
 #include "WidgetComponent.h"
 #include "VREditorWidgetComponent.h"
+#include "VREditorInteractor.h"
 #include "SLevelViewport.h"	// For tab manager tricks
 
 
@@ -14,33 +15,22 @@ namespace VREd
 {
 	static FAutoConsoleVariable UIFadeSpeed( TEXT( "VREd.UIFadeSpeed" ), 6.0f, TEXT( "How fast UI should fade in and out" ) );
 	static FAutoConsoleVariable UIBrightness( TEXT( "VREd.UIBrightness" ), 0.3f, TEXT( "How bright the UI should be" ) );
-	static FAutoConsoleVariable UIOnHandRotationOffset( TEXT( "VREd.UIOnHandRotationOffset" ), 45.0f, TEXT( "Rotation offset for UI that's docked to your hand, to make it more comfortable to hold" ) );
 }
 
 
 AVREditorFloatingUI::AVREditorFloatingUI()
-	: LocalRotation( FRotator( 90.0f, 180.0f, 0.0f ) ),
-	  RelativeOffset( FVector::ZeroVector ),
-	  SlateWidget( nullptr ),
-	  UserWidget( nullptr ),
-	  WidgetComponent( nullptr ),
-	  Scale( 1.0f ),
-	  Resolution( 0, 0 ),
-	  Owner( nullptr ),
-	  DockedTo( EDockedTo::Nothing ),
-	  PreviousDockedTo( EDockedTo::Nothing ),
-	  UserWidgetClass( nullptr ),
-	  bShouldBeVisible(),
-	  FadeAlpha( 1.0f ),
-	  bCollisionOnShowUI( true ),
-	  FadeDelay( 0.0f ),
-	  InitialScale( 1.0f ),
-	  bIsMoving( false ),
-	  MoveToTransform( FTransform() ),
-	  StartMoveToTransform( FTransform() ),
-	  MoveToAlpha( 0.0f ),
-	  MoveToTime( 0.0f ),
-	  MoveToResultDock( EDockedTo::Nothing )
+	: Super(),
+	SlateWidget( nullptr ),
+	UserWidget( nullptr ),
+	WidgetComponent( nullptr ),
+	Resolution( 0, 0 ),
+	Owner( nullptr ),
+	UserWidgetClass( nullptr ),
+	bShouldBeVisible(),
+	FadeAlpha( 1.0f ),
+	bCollisionOnShowUI( true ),
+	FadeDelay( 0.0f ),
+	InitialScale( 1.0f )
 {
 	const bool bTransient = true;
 	USceneComponent* SceneComponent = CreateDefaultSubobject<USceneComponent>( TEXT( "SceneComponent" ), bTransient );
@@ -108,9 +98,10 @@ void AVREditorFloatingUI::SetupWidgetComponent()
 	UpdateTransformIfDocked();
 }
 
-void AVREditorFloatingUI::SetSlateWidget( FVREditorUISystem& InitOwner, const TSharedRef<SWidget>& InitSlateWidget, const FIntPoint InitResolution, const float InitScale, const EDockedTo InitDockedTo )
+void AVREditorFloatingUI::SetSlateWidget( UVREditorUISystem& InitOwner, const TSharedRef<SWidget>& InitSlateWidget, const FIntPoint InitResolution, const float InitScale, const EDockedTo InitDockedTo )
 {
 	Owner = &InitOwner;
+	SetVRMode( &Owner->GetOwner() );
 
 	SlateWidget = InitSlateWidget;
 
@@ -120,15 +111,16 @@ void AVREditorFloatingUI::SetSlateWidget( FVREditorUISystem& InitOwner, const TS
 	Scale = InitScale;
 	InitialScale = Scale;
 
-	DockedTo = InitDockedTo;
+	SetDockedTo( InitDockedTo );
 
 	SetupWidgetComponent();
 }
 
 
-void AVREditorFloatingUI::SetUMGWidget( FVREditorUISystem& InitOwner, TSubclassOf<UVREditorBaseUserWidget> InitUserWidgetClass, const FIntPoint InitResolution, const float InitScale, const EDockedTo InitDockedTo )
+void AVREditorFloatingUI::SetUMGWidget( UVREditorUISystem& InitOwner, TSubclassOf<UVREditorBaseUserWidget> InitUserWidgetClass, const FIntPoint InitResolution, const float InitScale, const EDockedTo InitDockedTo )
 {
 	Owner = &InitOwner;
+	SetVRMode( &Owner->GetOwner() );
 
 	check( InitUserWidgetClass != nullptr );
 	UserWidgetClass = InitUserWidgetClass;
@@ -139,12 +131,11 @@ void AVREditorFloatingUI::SetUMGWidget( FVREditorUISystem& InitOwner, TSubclassO
 	Scale = InitScale;
 	InitialScale = Scale;
 
-	DockedTo = InitDockedTo;
+	SetDockedTo( InitDockedTo );
 
 	SetupWidgetComponent();
 }
 
-
 void AVREditorFloatingUI::Destroyed()
 {
 	if( WidgetComponent != nullptr )
@@ -168,23 +159,18 @@ void AVREditorFloatingUI::Destroyed()
 	Super::Destroyed();
 }
 
-
-void AVREditorFloatingUI::TickManually( float DeltaTime )
+void AVREditorFloatingUI::SetTransform( const FTransform& Transform )
 {
-	Super::Tick( DeltaTime );
+	const FVector AnimatedScale = CalculateAnimatedScale();
+
+	FTransform AnimatedTransform = Transform;
+	AnimatedTransform.SetScale3D( AnimatedTransform.GetScale3D() * AnimatedScale );
 
-	// Update fading state
-	UpdateFadingState( DeltaTime );
+	const float Aspect = ( float ) Resolution.X / ( float ) Resolution.Y;
 
-	if ( bIsMoving )
-	{
-		TickMoveTo( DeltaTime );
-	}
-	else
-	{
-		// Update transform
-		UpdateTransformIfDocked();
-	}
+	RootComponent->SetWorldLocation( AnimatedTransform.GetLocation() );
+	RootComponent->SetWorldRotation( AnimatedTransform.GetRotation() );
+	WidgetComponent->SetWorldScale3D( FVector( 1.0f / AnimatedTransform.GetScale3D().X, 1.0f / ( float ) Resolution.X, 1.0f / ( float ) Resolution.Y / Aspect ) * AnimatedTransform.GetScale3D() );
 }
 
 
@@ -238,51 +224,6 @@ void AVREditorFloatingUI::UpdateFadingState( const float DeltaTime )
 		WidgetComponent->SetTintColorAndOpacity( FLinearColor( UIBrightness, UIBrightness, UIBrightness ).CopyWithNewOpacity( FadeAlpha ) );
 	}
 }
-	
-
-void AVREditorFloatingUI::UpdateTransformIfDocked()
-{
-	if( DockedTo != EDockedTo::Nothing && DockedTo != EDockedTo::Dragging )
-	{
-		FTransform NewTransform;
-		
-		if( DockedTo == EDockedTo::LeftHand )
-		{
-			const bool bOnArm = false;
-			NewTransform = MakeUITransformLockedToHand( VREditorConstants::LeftHandIndex, bOnArm );
-		}
-		else if( DockedTo == EDockedTo::RightHand )
-		{
-			const bool bOnArm = false;
-			NewTransform = MakeUITransformLockedToHand( VREditorConstants::RightHandIndex, bOnArm );
-		}
-		else if( DockedTo == EDockedTo::LeftArm )
-		{
-			const bool bOnArm = true;
-			NewTransform = MakeUITransformLockedToHand( VREditorConstants::LeftHandIndex, bOnArm );
-		}
-		else if( DockedTo == EDockedTo::RightArm )
-		{
-			const bool bOnArm = true;
-			NewTransform = MakeUITransformLockedToHand( VREditorConstants::RightHandIndex, bOnArm );
-		}
-		else if( DockedTo == EDockedTo::Room)
-		{
-			NewTransform = MakeUITransformLockedToRoom();
-		}
-		else if( DockedTo == EDockedTo::Custom )
-		{
-			NewTransform = MakeCustomUITransform();
-		}
-		else
-		{
-			check( 0 );
-		}
-
-		SetTransform( NewTransform );
-	}
-}
-
 
 FVector AVREditorFloatingUI::CalculateAnimatedScale() const
 {
@@ -300,32 +241,6 @@ FVector AVREditorFloatingUI::CalculateAnimatedScale() const
 	return AnimatedScale;
 }
 
-
-void AVREditorFloatingUI::SetTransform( const FTransform& Transform )
-{
-	const FVector AnimatedScale = CalculateAnimatedScale();
-	
-	FTransform AnimatedTransform = Transform;
-	AnimatedTransform.SetScale3D( AnimatedTransform.GetScale3D() * AnimatedScale );
-
-	const float Aspect = (float)Resolution.X / (float)Resolution.Y;
-
-	RootComponent->SetWorldLocation( AnimatedTransform.GetLocation() );
-	RootComponent->SetWorldRotation( AnimatedTransform.GetRotation() );
-	WidgetComponent->SetWorldScale3D( FVector( 1.0f / AnimatedTransform.GetScale3D().X, 1.0f / (float)Resolution.X, 1.0f / (float)Resolution.Y / Aspect ) * AnimatedTransform.GetScale3D() );
-}
-
-
-void AVREditorFloatingUI::SetRelativeOffset( const FVector& InRelativeOffset )
-{
-	RelativeOffset = InRelativeOffset;
-}
-
-void AVREditorFloatingUI::SetLocalRotation( const FRotator& InLocalRotation )
-{
-	LocalRotation = InLocalRotation;
-}
-
 void AVREditorFloatingUI::SetCollisionOnShow( const bool bInCollisionOnShow )
 {
 	bCollisionOnShowUI = bInCollisionOnShow;
@@ -336,88 +251,6 @@ float AVREditorFloatingUI::GetInitialScale() const
 	return InitialScale;
 }
 
-FTransform AVREditorFloatingUI::MakeUITransformLockedToHand( const int32 HandIndex, const bool bOnArm )
-{
-	return MakeUITransformLockedToHand( HandIndex, bOnArm, RelativeOffset, LocalRotation );
-}
-
-FTransform AVREditorFloatingUI::MakeUITransformLockedToHand( const int32 HandIndex, const bool bOnArm, const FVector& InRelativeOffset, const FRotator& InLocalRotation )
-{
-	const float WorldScaleFactor = Owner->GetOwner().GetWorldScaleFactor();
-
-	FTransform UIToHandTransform( InLocalRotation, InRelativeOffset * WorldScaleFactor );
-	if (!bOnArm)
-	{
-		UIToHandTransform *= FTransform( FRotator( VREd::UIOnHandRotationOffset->GetFloat(), 0.0f, 0.0f ).Quaternion(), FVector::ZeroVector );
-	}
-
-	const FTransform HandToWorldTransform = Owner->GetOwner().GetVirtualHand( HandIndex ).Transform;
-
-	FTransform UIToWorldTransform = UIToHandTransform * HandToWorldTransform;
-	UIToWorldTransform.SetScale3D( FVector( Scale * WorldScaleFactor ) );
-
-	return UIToWorldTransform;
-}
-
-void AVREditorFloatingUI::MoveTo( const FTransform& ResultTransform, const float TotalMoveToTime, const EDockedTo ResultDock )
-{
-	MoveToTime = TotalMoveToTime;
-	bIsMoving = true;
-	StartMoveToTransform = GetActorTransform();
-	MoveToTransform = ResultTransform;
-	MoveToAlpha = 0.0f;
-	MoveToResultDock = ResultDock;
-}
-
-void AVREditorFloatingUI::StopMoveTo()
-{
-	bIsMoving = false;
-	SetTransform( MoveToTransform );
-	MoveToTransform = FTransform();
-}
-
-FTransform AVREditorFloatingUI::MakeUITransformLockedToRoom()
-{
-	const float WorldScaleFactor = Owner->GetOwner().GetWorldScaleFactor();
-
-	const FTransform UIToRoomTransform( LocalRotation, RelativeOffset * WorldScaleFactor );
-
-	const FTransform RoomToWorldTransform = Owner->GetOwner().GetRoomTransform();
-
-	FTransform UIToWorldTransform = UIToRoomTransform * RoomToWorldTransform;
-	UIToWorldTransform.SetScale3D( FVector( Scale * WorldScaleFactor ) );
-
-	return UIToWorldTransform;
-}
-
-void AVREditorFloatingUI::TickMoveTo( const float DeltaTime )
-{
-	const float WorldScaleFactor = Owner->GetOwner().GetWorldScaleFactor();
-
-	MoveToAlpha += DeltaTime;
-	const float LerpTime = MoveToTime;
-	if (MoveToAlpha > MoveToTime)
-	{
-		MoveToAlpha = MoveToTime;
-		bIsMoving = false;
-		SetDockedTo( MoveToResultDock );
-	}
-
-	const float CurrentALpha = MoveToAlpha / LerpTime;
-	const FVector NewLocation = FMath::Lerp( StartMoveToTransform.GetLocation(), MoveToTransform.GetLocation(), CurrentALpha );
-	const FQuat NewRotation = FMath::Lerp( StartMoveToTransform.GetRotation(), MoveToTransform.GetRotation(), CurrentALpha );
-	FTransform NewTransform( NewRotation, NewLocation, FVector( Scale * WorldScaleFactor ) );
-
-	SetTransform( NewTransform );
-}
-
-void AVREditorFloatingUI::SetDockedTo( const EDockedTo NewDockedTo )
-{
-	PreviousDockedTo = DockedTo;
-	DockedTo = NewDockedTo;
-	UpdateTransformIfDocked();
-}
-
 
 void AVREditorFloatingUI::ShowUI( const bool bShow, const bool bAllowFading, const float InitFadeDelay )
 {
diff --git a/Engine/Source/Editor/VREditor/VREditorFloatingUI.h b/Engine/Source/Editor/VREditor/VREditorFloatingUI.h
--- a/Engine/Source/Editor/VREditor/VREditorFloatingUI.h
+++ b/Engine/Source/Editor/VREditor/VREditorFloatingUI.h
@@ -2,7 +2,7 @@
 
 #pragma once
 
-#include "GameFramework/Actor.h"
+#include "VREditorBaseActor.h"
 #include "VREditorFloatingUI.generated.h"
 
 
@@ -10,7 +10,7 @@
  * Represents an interactive floating UI panel in the VR Editor
  */
 UCLASS()
-class AVREditorFloatingUI : public AActor
+class AVREditorFloatingUI : public AVREditorBaseActor
 {
 	GENERATED_BODY()
 
@@ -19,24 +19,11 @@ public:
 	/** Default constructor which sets up safe defaults */
 	AVREditorFloatingUI();
 
-	/** Possible UI attachment points */
-	enum class EDockedTo
-	{
-		Nothing,
-		LeftHand,
-		RightHand,
-		LeftArm,
-		RightArm,
-		Room,
-		Custom,
-		Dragging,
-	};
-
 	/** Creates a FVREditorFloatingUI using a Slate widget, and sets up safe defaults */
-	void SetSlateWidget( class FVREditorUISystem& InitOwner, const TSharedRef<SWidget>& InitSlateWidget, const FIntPoint InitResolution, const float InitScale, const EDockedTo InitDockedTo );
+	void SetSlateWidget( class UVREditorUISystem& InitOwner, const TSharedRef<SWidget>& InitSlateWidget, const FIntPoint InitResolution, const float InitScale, const EDockedTo InitDockedTo );
 
 	/** Creates a FVREditorFloatingUI using a UMG user widget, and sets up safe defaults */
-	void SetUMGWidget( class FVREditorUISystem& InitOwner, class TSubclassOf<class UVREditorBaseUserWidget> InitUserWidgetClass, const FIntPoint InitResolution, const float InitScale, const EDockedTo InitDockedTo );
+	void SetUMGWidget( class UVREditorUISystem& InitOwner, class TSubclassOf<class UVREditorBaseUserWidget> InitUserWidgetClass, const FIntPoint InitResolution, const float InitScale, const EDockedTo InitDockedTo );
 
 	/** @return Returns true if the UI is visible (or wants to be visible -- it might be transitioning */
 	bool IsUIVisible() const
@@ -47,21 +34,6 @@ public:
 	/** Shows or hides the UI (also enables collision, and performs a transition effect) */
 	void ShowUI( const bool bShow, const bool bAllowFading = true, const float InitFadeDelay = 0.0f );
 
-	/** Returns what we're docked to, if anything */
-	EDockedTo GetDockedTo() const
-	{
-		return DockedTo;
-	}
-
-	/** Returns what we were previously docked to */
-	EDockedTo GetPreviouslyDockedTo() const
-	{
-		return PreviousDockedTo;
-	}
-
-	/** Sets what this UI is docked to */
-	void SetDockedTo( const EDockedTo NewDockedTo );
-
 	/** Returns the widget component for this UI, or nullptr if not spawned right now */
 	class UVREditorWidgetComponent* GetWidgetComponent()
 	{
@@ -77,25 +49,21 @@ public:
 	/** Sets a new size for the UI */
 	void SetScale( const float NewSize );
 
+	/** Sets the UI transform */
+	virtual void SetTransform( const FTransform& Transform ) override;
+
 	/** AActor overrides */
 	virtual void Destroyed() override;
 
-	/** Called by the UI system to tick this UI every frame.  We can't use AActor::Tick() because this actor can exist in 
-	    the editor's world, which never ticks. */
-	virtual void TickManually( float DeltaTime );
-
-	/** Sets the transform for this UI */
-	void SetTransform( const FTransform& Transform );
-
 	/** Returns the owner of this object */
-	FVREditorUISystem& GetOwner()
+	UVREditorUISystem& GetOwner()
 	{
 		check( Owner != nullptr );
 		return *Owner;
 	}
 
 	/** Returns the owner of this object (const) */
-	const FVREditorUISystem& GetOwner() const
+	const UVREditorUISystem& GetOwner() const
 	{
 		check( Owner != nullptr );
 		return *Owner;
@@ -108,62 +76,24 @@ public:
 		return CastChecked<T>( UserWidget );
 	}
 
-	/** Sets the relative offset of the UI */
-	void SetRelativeOffset( const FVector& InRelativeOffset );
-
-	/** Sets the local rotation of the UI */
-	void SetLocalRotation( const FRotator& InRelativeRotation );
-
 	/** Sets if the collision is on when showing the UI*/
 	void SetCollisionOnShow( const bool bInCollisionOnShow );
 
 	/** Gets the initial size of this UI */
 	float GetInitialScale() const;
 
-	/** Get a relative transform from the hand */
-	FTransform MakeUITransformLockedToHand( const int32 HandIndex, const bool bOnArm, const FVector& InRelativeOffset, const FRotator& InLocalRotation );
-
-	/** Start moving towards a transform */
-	void MoveTo( const FTransform& ResultTransform, const float TotalMoveToTime, const EDockedTo ResultDock );
-
-	/** Abort moving to transform and move instant to MoveToTransform */
-	void StopMoveTo();
-
 protected:
 
 	/** Returns a scale to use for this widget that takes into account animation */
 	FVector CalculateAnimatedScale() const;
 
-	/** Called when DockTo state is custom, needs to be implemented by derived classes */
-	virtual FTransform MakeCustomUITransform() { return FTransform(); };
-
 	/** Called to finish setting everything up, after a widget has been assigned */
 	virtual void SetupWidgetComponent();
 
 private:
 
 	/** Called after spawning, and every Tick, to update opacity of the widget */
-	void UpdateFadingState( const float DeltaTime );
-
-	/** Called every tick to keep the UI position up to date */
-	void UpdateTransformIfDocked();
-
-	/** Given a hand to lock to, returns a transform to place UI at that hand's location and orientation */
-	FTransform MakeUITransformLockedToHand( const int32 HandIndex, const bool bOnArm);
-
-	/** Create a room transform with the RelativeOffset and LocalRotation */
-	FTransform MakeUITransformLockedToRoom();
-
-	/** Updates the lerp movement */
-	void TickMoveTo( const float DeltaTime );
-
-public:
-
-	/** Local rotation of the UI */
-	FRotator LocalRotation;
-
-	/** Relative offset of the UI */
-	FVector RelativeOffset;
+	virtual void UpdateFadingState( const float DeltaTime );
 
 protected:
 	/** Slate widget we're drawing, or null if we're drawing a UMG user widget */
@@ -176,9 +106,6 @@ protected:
 	/** When in a spawned state, this is the widget component to represent the widget */
 	UPROPERTY()
 	class UVREditorWidgetComponent* WidgetComponent;
-	
-	/** How big the UI should be */
-	float Scale;
 
 	/** Resolution we should draw this UI at, regardless of scale */
 	FIntPoint Resolution;
@@ -186,13 +113,7 @@ protected:
 private:
 
 	/** Owning object */
-	class FVREditorUISystem* Owner;
-
-	/** What the UI is attached to */
-	EDockedTo DockedTo;
-
-	/** What was the UI previously attached to */
-	EDockedTo PreviousDockedTo;
+	class UVREditorUISystem* Owner;
 
 	/** Class of the UMG widget we're spawning */
 	UPROPERTY()
@@ -213,23 +134,5 @@ private:
 
 	/** The starting scale of this UI */
 	float InitialScale;
-
-	/** If the UI is currently moving */
-	bool bIsMoving;
-
-	/** The end transform to move to */
-	FTransform MoveToTransform;
-
-	/** The actor transform when started moving */
-	FTransform StartMoveToTransform;
-
-	/** Current lerp alpha */
-	float MoveToAlpha;
-
-	/** Total time to move to the end transform */
-	float MoveToTime;
-
-	/** Dock state to switch to when finished moving */
-	EDockedTo MoveToResultDock;
 };
 
diff --git a/Engine/Source/Editor/VREditor/VREditorInteractor.cpp b/Engine/Source/Editor/VREditor/VREditorInteractor.cpp
new file mode 100644
--- /dev/null
+++ b/Engine/Source/Editor/VREditor/VREditorInteractor.cpp
@@ -0,0 +1,214 @@
+// Copyright 1998-2015 Epic Games, Inc. All Rights Reserved.
+
+#include "VREditorModule.h"
+#include "VREditorInteractor.h"
+#include "VREditorWorldInteraction.h"
+
+#include "VREditorMode.h"
+#include "VREditorFloatingText.h"
+
+#include "Interactables/VREditorButton.h"
+#include "VREditorDockableWindow.h"
+
+#define LOCTEXT_NAMESPACE "VREditor"
+
+UVREditorInteractor::UVREditorInteractor( const FObjectInitializer& Initializer ) :
+	Super( Initializer ),
+	VRMode( nullptr ),
+	bIsModifierPressed( false ),
+	SelectAndMoveTriggerValue( 0.0f ),
+	bHasUIInFront( false ),
+	bHasUIOnForearm( false ),
+	bIsClickingOnUI( false ),
+	bIsRightClickingOnUI( false ),
+	bIsHoveringOverUI( false ),
+	UIScrollVelocity( 0.0f ),
+	LastClickReleaseTime( 0.0f ),
+	bWantHelpLabels( false ),
+	HelpLabelShowOrHideStartTime( FTimespan::MinValue() )
+{
+
+}
+
+UVREditorInteractor::~UVREditorInteractor()
+{
+	Shutdown();
+}
+
+void UVREditorInteractor::Init( FVREditorMode* InVRMode )
+{
+	VRMode = InVRMode;
+	KeyToActionMap.Reset();
+}
+
+void UVREditorInteractor::Shutdown()
+{
+	for ( auto& KeyAndValue : HelpLabels )
+	{
+		AFloatingText* FloatingText = KeyAndValue.Value;
+		GetVRMode().DestroyTransientActor( FloatingText );
+	}
+
+	HelpLabels.Empty();	
+
+	VRMode = nullptr;
+
+	Super::Shutdown();
+}
+
+void UVREditorInteractor::Tick( const float DeltaTime )
+{
+	Super::Tick( DeltaTime );
+}
+
+FHitResult UVREditorInteractor::GetHitResultFromLaserPointer( TArray<AActor*>* OptionalListOfIgnoredActors /*= nullptr*/, const bool bIgnoreGizmos /*= false*/,
+	TArray<UClass*>* ObjectsInFrontOfGizmo /*= nullptr */, const bool bEvenIfBlocked /*= false */, const float LaserLengthOverride /*= 0.0f */ )
+{
+	static TArray<AActor*> IgnoredActors;
+	IgnoredActors.Reset();
+	if ( OptionalListOfIgnoredActors == nullptr )
+	{
+		OptionalListOfIgnoredActors = &IgnoredActors;
+	}
+
+	OptionalListOfIgnoredActors->Add( GetVRMode().GetAvatarMeshActor() );
+	
+	// Ignore UI widgets too
+	if ( GetDraggingMode() == EViewportInteractionDraggingMode::ActorsAtLaserImpact )
+	{
+		for( TActorIterator<AVREditorFloatingUI> UIActorIt( WorldInteraction->GetViewportWorld() ); UIActorIt; ++UIActorIt )
+		{
+			OptionalListOfIgnoredActors->Add( *UIActorIt );
+		}
+	}
+
+	static TArray<UClass*> PriorityOverGizmoObjects;
+	PriorityOverGizmoObjects.Reset();
+	if ( ObjectsInFrontOfGizmo == nullptr )
+	{
+		ObjectsInFrontOfGizmo = &PriorityOverGizmoObjects;
+	}
+
+	ObjectsInFrontOfGizmo->Add( AVREditorButton::StaticClass() );
+	ObjectsInFrontOfGizmo->Add( AVREditorDockableWindow::StaticClass() );
+
+	return UViewportInteractor::GetHitResultFromLaserPointer( OptionalListOfIgnoredActors, bIgnoreGizmos, ObjectsInFrontOfGizmo, bEvenIfBlocked, LaserLengthOverride );
+}
+
+void UVREditorInteractor::ResetHoverState( const float DeltaTime )
+{
+	bIsHoveringOverUI = false;
+}
+
+void UVREditorInteractor::OnStartDragging( UActorComponent* ClickedComponent, const FVector& HitLocation, const bool bIsPlacingActors )
+{
+	// If the user is holding down the modifier key, go ahead and duplicate the selection first.  This needs to
+	// happen before we create our transformables, because duplicating actors will change which actors are
+	// selected!  Don't do this if we're placing objects right now though.
+	if ( bIsModifierPressed && !bIsPlacingActors )
+	{
+		WorldInteraction->Duplicate();
+	}
+}
+
+float UVREditorInteractor::GetSlideDelta()
+{
+	return 0.0f;
+}
+
+bool UVREditorInteractor::IsHoveringOverUI() const
+{
+	return bIsHoveringOverUI;
+}
+
+void UVREditorInteractor::SetHasUIInFront( const bool bInHasUIInFront )
+{
+	bHasUIInFront = bInHasUIInFront;
+}
+
+bool UVREditorInteractor::HasUIInFront() const
+{
+	return bHasUIInFront;
+}
+
+void UVREditorInteractor::SetHasUIOnForearm( const bool bInHasUIOnForearm )
+{
+	bHasUIOnForearm = bInHasUIOnForearm;
+}
+
+bool UVREditorInteractor::HasUIOnForearm() const
+{
+	return bHasUIOnForearm;
+}
+
+UWidgetComponent* UVREditorInteractor::GetHoveringOverWidgetComponent() const
+{
+	return InteractorData.HoveringOverWidgetComponent;
+}
+
+void UVREditorInteractor::SetHoveringOverWidgetComponent( UWidgetComponent* NewHoveringOverWidgetComponent )
+{
+	InteractorData.HoveringOverWidgetComponent = NewHoveringOverWidgetComponent;
+}
+
+bool UVREditorInteractor::IsModifierPressed() const
+{
+	return bIsModifierPressed;
+}
+
+void UVREditorInteractor::SetIsClickingOnUI( const bool bInIsClickingOnUI )
+{
+	bIsClickingOnUI = bInIsClickingOnUI;
+}
+
+bool UVREditorInteractor::IsClickingOnUI() const
+{
+	return bIsClickingOnUI;
+}
+
+void UVREditorInteractor::SetIsHoveringOverUI( const bool bInIsHoveringOverUI )
+{
+	bIsHoveringOverUI = bInIsHoveringOverUI;
+}
+
+void UVREditorInteractor::SetIsRightClickingOnUI( const bool bInIsRightClickingOnUI )
+{
+	bIsRightClickingOnUI = bInIsRightClickingOnUI;
+}
+
+bool UVREditorInteractor::IsRightClickingOnUI() const
+{
+	return bIsRightClickingOnUI;
+}
+
+void UVREditorInteractor::SetLastClickReleaseTime( const double InLastClickReleaseTime )
+{
+	LastClickReleaseTime = InLastClickReleaseTime;
+}
+
+double UVREditorInteractor::GetLastClickReleaseTime() const
+{
+	return LastClickReleaseTime;
+}
+
+void UVREditorInteractor::SetUIScrollVelocity( const float InUIScrollVelocity )
+{
+	UIScrollVelocity = InUIScrollVelocity;
+}
+
+float UVREditorInteractor::GetUIScrollVelocity() const
+{
+	return UIScrollVelocity;
+}
+
+float UVREditorInteractor::GetSelectAndMoveTriggerValue() const
+{
+	return SelectAndMoveTriggerValue;
+}
+
+bool UVREditorInteractor::GetIsLaserBlocked()
+{
+	return bHasUIInFront;
+}
+
+#undef LOCTEXT_NAMESPACE
\ No newline at end of file
