Copying //UE4/Dev-Editor to //UE4/Dev-Main (Source: //UE4/Dev-Editor @ 3082391)

diff --git a/Engine/Source/Editor/MeshPaint/Private/MeshPaintEdMode.cpp b/Engine/Source/Editor/MeshPaint/Private/MeshPaintEdMode.cpp
--- a/Engine/Source/Editor/MeshPaint/Private/MeshPaintEdMode.cpp
+++ b/Engine/Source/Editor/MeshPaint/Private/MeshPaintEdMode.cpp
@@ -12,6 +12,7 @@
 #include "RawMesh.h"
 #include "Editor/UnrealEd/Public/ObjectTools.h"
 #include "AssetToolsModule.h"
+#include "AssetRegistryModule.h"
 #include "EditorSupportDelegates.h"
 #include "EditorReimportHandler.h"
 
@@ -162,13 +163,26 @@ void FEdModeMeshPaint::Enter()
 
 	}
 
+	// Catch when objects are replaced when a construction script is rerun
+	GEditor->OnObjectsReplaced().AddSP(this, &FEdModeMeshPaint::OnObjectsReplaced);
+
+	// Hook into pre/post world save, so that the original collision volumes can be temporarily reinstated
+	FEditorDelegates::PreSaveWorld.AddSP(this, &FEdModeMeshPaint::OnPreSaveWorld);
+	FEditorDelegates::PostSaveWorld.AddSP(this, &FEdModeMeshPaint::OnPostSaveWorld);
+
 	// Catch assets if they are about to be (re)imported
-	FEditorDelegates::OnAssetPreImport.AddSP(this, &FEdModeMeshPaint::OnPreImportAsset);
-	FReimportManager::Instance()->OnPreReimport().AddSP(this, &FEdModeMeshPaint::OnPreReimportAsset);
+	FEditorDelegates::OnAssetPostImport.AddSP(this, &FEdModeMeshPaint::OnPostImportAsset);
+	FReimportManager::Instance()->OnPostReimport().AddSP(this, &FEdModeMeshPaint::OnPostReimportAsset);
+
+	FAssetRegistryModule& AssetRegistryModule = FModuleManager::LoadModuleChecked<FAssetRegistryModule>(TEXT("AssetRegistry"));
+	AssetRegistryModule.Get().OnAssetRemoved().AddSP(this, &FEdModeMeshPaint::OnAssetRemoved);
 
 	// Initialize adapter globals
 	FMeshPaintAdapterFactory::InitializeAdapterGlobals();
 
+	// Create adapters for currently selected mesh components
+	CreateGeometryAdaptersForSelectedComponents();
+
 	if (!Toolkit.IsValid())
 	{
 		Toolkit = MakeShareable(new FMeshPaintToolKit(this));
@@ -295,8 +309,14 @@ void FEdModeMeshPaint::Exit()
 	}
 
 	// Unbind delegates
-	FReimportManager::Instance()->OnPreReimport().RemoveAll(this);
-	FEditorDelegates::OnAssetPreImport.RemoveAll(this);
+	FAssetRegistryModule& AssetRegistryModule = FModuleManager::LoadModuleChecked<FAssetRegistryModule>(TEXT("AssetRegistry"));
+	AssetRegistryModule.Get().OnAssetRemoved().RemoveAll(this);
+
+	FReimportManager::Instance()->OnPostReimport().RemoveAll(this);
+	FEditorDelegates::OnAssetPostImport.RemoveAll(this);
+	FEditorDelegates::PreSaveWorld.RemoveAll(this);
+	FEditorDelegates::PostSaveWorld.RemoveAll(this);
+	GEditor->OnObjectsReplaced().RemoveAll(this);
 
 	// Call parent implementation
 	FEdMode::Exit();
@@ -988,37 +1008,47 @@ bool FEdModeMeshPaint::PaintVertex( const FVector& InVertexPosition,
 	return false;
 }
 
-IMeshPaintGeometryAdapter* FEdModeMeshPaint::FindOrAddGeometryAdapter(UMeshComponent* MeshComponent)
+void FEdModeMeshPaint::CreateGeometryAdaptersForSelectedComponents()
 {
-	TSharedPtr<IMeshPaintGeometryAdapter>* MeshAdapter = ComponentToAdapterMap.Find(MeshComponent);
-	if (!MeshAdapter)
+	TArray<UMeshComponent*> SelectedComponents = GetSelectedMeshComponents();
+	for (UMeshComponent* SelectedComponent : SelectedComponents)
 	{
-		// If this component hasn't yet been seen, make an adapter for it and add it to the map
-		TSharedPtr<IMeshPaintGeometryAdapter> NewAdapter = FMeshPaintAdapterFactory::CreateAdapterForMesh(MeshComponent, PaintingMeshLODIndex, /*TODO: Shouldn't be part of the construction contract: FMeshPaintSettings::Get().UVChannel*/ 0);
-		if (NewAdapter.IsValid())
-		{
-			ComponentToAdapterMap.Add(MeshComponent, NewAdapter);
-			NewAdapter->OnAdded();
-		}
-		return NewAdapter.Get();
+		AddGeometryAdapter(SelectedComponent);
 	}
-	else
+}
+
+IMeshPaintGeometryAdapter* FEdModeMeshPaint::AddGeometryAdapter(UMeshComponent* MeshComponent)
+{
+	check(MeshComponent);
+	check(!ComponentToAdapterMap.Contains(MeshComponent));
+	TSharedPtr<IMeshPaintGeometryAdapter> NewAdapter = FMeshPaintAdapterFactory::CreateAdapterForMesh(MeshComponent, PaintingMeshLODIndex, /*TODO: Shouldn't be part of the construction contract: FMeshPaintSettings::Get().UVChannel*/ 0);
+	if (NewAdapter.IsValid())
 	{
-		// Use existing adapter
-		return MeshAdapter->Get();
+		ComponentToAdapterMap.Add(MeshComponent, NewAdapter);
+		NewAdapter->OnAdded();
 	}
+	return NewAdapter.Get();
 }
 
-void FEdModeMeshPaint::CleanStaleGeometryAdapters(const TArray<UMeshComponent*>& ValidComponents)
+void FEdModeMeshPaint::RemoveGeometryAdapter(UMeshComponent* MeshComponent)
 {
-	// Remove any stale components from the map
-	for (auto It = ComponentToAdapterMap.CreateIterator(); It; ++It)
+	check(MeshComponent);
+	TSharedPtr<IMeshPaintGeometryAdapter>* MeshAdapter = ComponentToAdapterMap.Find(MeshComponent);
+	check(MeshAdapter);
+	(*MeshAdapter)->OnRemoved();
+	verify(ComponentToAdapterMap.Remove(MeshComponent) == 1);
+}
+
+IMeshPaintGeometryAdapter* FEdModeMeshPaint::FindGeometryAdapter(UMeshComponent* MeshComponent)
+{
+	TSharedPtr<IMeshPaintGeometryAdapter>* MeshAdapter = ComponentToAdapterMap.Find(MeshComponent);
+	if (!MeshAdapter)
 	{
-		if (!It.Value()->IsValid() || !ValidComponents.Contains(It.Key()))
-		{
-			It.Value()->OnRemoved();
-			It.RemoveCurrent();
-		}
+		return nullptr;
+	}
+	else
+	{
+		return MeshAdapter->Get();
 	}
 }
 
@@ -1032,24 +1062,58 @@ void FEdModeMeshPaint::RemoveAllGeometryAdapters()
 	}
 }
 
-void FEdModeMeshPaint::OnPreImportAsset(UFactory* Factory, UClass* Class, UObject* Object, const FName& Name, const TCHAR* Type)
+void FEdModeMeshPaint::OnPreSaveWorld(uint32 SaveFlags, UWorld* World)
 {
-	if (Class->IsChildOf(UStaticMesh::StaticClass()))
+	RemoveAllGeometryAdapters();
+}
+
+void FEdModeMeshPaint::OnPostSaveWorld(uint32 SaveFlags, UWorld* World, bool bSuccess)
+{
+	CreateGeometryAdaptersForSelectedComponents();
+}
+
+void FEdModeMeshPaint::OnPostImportAsset(UFactory* Factory, UObject* Object)
+{
+	if (Object->IsA(UStaticMesh::StaticClass()))
 	{
-		// Remove all geometry adapters to force them all to be recached next time they're required
 		RemoveAllGeometryAdapters();
+		CreateGeometryAdaptersForSelectedComponents();
 	}
 }
 
-void FEdModeMeshPaint::OnPreReimportAsset(UObject* Object)
+void FEdModeMeshPaint::OnPostReimportAsset(UObject* Object, bool bSuccess)
 {
-	if (Object->IsA(UStaticMesh::StaticClass()))
+	if (bSuccess && Object->IsA(UStaticMesh::StaticClass()))
 	{
-		// Remove all geometry adapters to force them all to be recached next time they're required
 		RemoveAllGeometryAdapters();
+		CreateGeometryAdaptersForSelectedComponents();
 	}
 }
 
+void FEdModeMeshPaint::OnAssetRemoved(const FAssetData& AssetData)
+{
+	RemoveAllGeometryAdapters();
+	CreateGeometryAdaptersForSelectedComponents();
+}
+
+void FEdModeMeshPaint::OnObjectsReplaced(const TMap<UObject*, UObject*>& OldToNewInstanceMap)
+{
+	for (const auto& OldToNew : OldToNewInstanceMap)
+	{
+		if (UMeshComponent* OldMeshComponent = Cast<UMeshComponent>(OldToNew.Key))
+		{
+			if (ComponentToAdapterMap.Contains(OldMeshComponent))
+			{
+				if (UMeshComponent* NewMeshComponent = Cast<UMeshComponent>(OldToNew.Value))
+				{
+					AddGeometryAdapter(NewMeshComponent);
+				}
+
+				RemoveGeometryAdapter(OldMeshComponent);
+			}
+		}
+	}
+}
 
 /** Paint the mesh that impacts the specified ray */
 void FEdModeMeshPaint::DoPaint( const FVector& InCameraOrigin,
@@ -1076,7 +1140,7 @@ void FEdModeMeshPaint::DoPaint( const FVector& InCameraOrigin,
 
 		for (UMeshComponent* MeshComponent : SelectedMeshComponents)
 		{
-			IMeshPaintGeometryAdapter* MeshAdapter = FindOrAddGeometryAdapter(MeshComponent);
+			IMeshPaintGeometryAdapter* MeshAdapter = FindGeometryAdapter(MeshComponent);
 			if (!MeshAdapter)
 			{
 				continue;
@@ -1106,9 +1170,6 @@ void FEdModeMeshPaint::DoPaint( const FVector& InCameraOrigin,
 			}
 		}
 
-		// Remove any stale components from the map
-		CleanStaleGeometryAdapters(SelectedMeshComponents);
-
 		if (BestTraceResult.GetComponent() != NULL)
 		{
 			// If we're using texture paint, just use the best trace result we found as we currently only
@@ -2610,6 +2671,7 @@ void FEdModeMeshPaint::PostUndo()
 	FEdMode::PostUndo();
 	bDoRestoreRenTargets = true;
 	RemoveAllGeometryAdapters();
+	CreateGeometryAdaptersForSelectedComponents();
 }
 
 /** Returns true if we need to force a render/update through based fill/copy */
@@ -2933,31 +2995,25 @@ bool FEdModeMeshPaint::Select( AActor* InActor, bool bInSelected )
 					}
 				}
 
-				GeomInfo->OnRemoved();
-				ensure(ComponentToAdapterMap.Remove(MeshComponent) == 1);
+				RemoveGeometryAdapter(MeshComponent);
 			}
 		}
 		else
 		{
 			if (!ComponentToAdapterMap.Contains(MeshComponent))
 			{
-				TSharedPtr<IMeshPaintGeometryAdapter> GeomInfo = FMeshPaintAdapterFactory::CreateAdapterForMesh(MeshComponent, PaintingMeshLODIndex, /*TODO: Shouldn't be part of the construction contract: FMeshPaintSettings::Get().UVChannel*/ 0);
-				if (GeomInfo.IsValid())
-				{
-					ComponentToAdapterMap.Add(MeshComponent, GeomInfo);
-					GeomInfo->OnAdded();
+				IMeshPaintGeometryAdapter* GeomInfo = AddGeometryAdapter(MeshComponent);
 
-					if (FMeshPaintSettings::Get().ResourceType == EMeshPaintResource::Texture)
-					{
-						SetAllTextureOverrides(*GeomInfo.Get(), MeshComponent);
-					}
-					else if (FMeshPaintSettings::Get().ResourceType == EMeshPaintResource::VertexColors)
+				if (FMeshPaintSettings::Get().ResourceType == EMeshPaintResource::Texture)
+				{
+					SetAllTextureOverrides(*GeomInfo, MeshComponent);
+				}
+				else if (FMeshPaintSettings::Get().ResourceType == EMeshPaintResource::VertexColors)
+				{
+					//Painting is done on LOD0 so force the mesh to render only LOD0.
+					ApplyOrRemoveForceBestLOD(*GeomInfo, MeshComponent, /*bApply=*/ true);
 					{
-						//Painting is done on LOD0 so force the mesh to render only LOD0.
-						ApplyOrRemoveForceBestLOD(*GeomInfo.Get(), MeshComponent, /*bApply=*/ true);
-						{
-							FComponentReregisterContext ReregisterContext(MeshComponent);
-						}
+						FComponentReregisterContext ReregisterContext(MeshComponent);
 					}
 				}
 			}
@@ -3898,14 +3954,12 @@ void FEdModeMeshPaint::ApplyOrRemoveForceBestLOD(bool bApply)
 
 	for (UMeshComponent* MeshComponent : SelectedMeshComponents)
 	{
-		IMeshPaintGeometryAdapter* MeshAdapter = FindOrAddGeometryAdapter(MeshComponent);
+		IMeshPaintGeometryAdapter* MeshAdapter = FindGeometryAdapter(MeshComponent);
 		if (MeshAdapter)
 		{
 			ApplyOrRemoveForceBestLOD(*MeshAdapter, MeshComponent, bApply);
 		}
 	}
-
-	CleanStaleGeometryAdapters(SelectedMeshComponents);
 }
 
 void FEdModeMeshPaint::ApplyOrRemoveForceBestLOD(const IMeshPaintGeometryAdapter& GeometryInfo, UMeshComponent* MeshComponent, bool bApply)
@@ -3924,14 +3978,12 @@ void FEdModeMeshPaint::ApplyVertexColorsToAllLODs()
 
 	for (UMeshComponent* MeshComponent : SelectedMeshComponents)
 	{
-		IMeshPaintGeometryAdapter* MeshAdapter = FindOrAddGeometryAdapter(MeshComponent);
+		IMeshPaintGeometryAdapter* MeshAdapter = FindGeometryAdapter(MeshComponent);
 		if (MeshAdapter)
 		{
 			ApplyVertexColorsToAllLODs(*MeshAdapter, MeshComponent);
 		}
 	}
-
-	CleanStaleGeometryAdapters(SelectedMeshComponents);
 }
 
 void FEdModeMeshPaint::ApplyVertexColorsToAllLODs(const IMeshPaintGeometryAdapter& GeometryInfo, UMeshComponent* InMeshComponent)
@@ -4450,7 +4502,7 @@ void FEdModeMeshPaint::UpdateTexturePaintTargetList()
 				for (const auto& MeshComponent : MeshComponents)
 				{
 					// Get the geometry adapter
-					IMeshPaintGeometryAdapter* MeshAdapter = FindOrAddGeometryAdapter(MeshComponent);
+					IMeshPaintGeometryAdapter* MeshAdapter = FindGeometryAdapter(MeshComponent);
 					if (MeshAdapter)
 					{
 						// We already know the material we are painting on, take it off the static mesh component
@@ -4841,14 +4893,12 @@ void FEdModeMeshPaint::Tick(FEditorViewportClient* ViewportClient,float DeltaTim
 
 		for (UMeshComponent* MeshComponent : SelectedMeshComponents)
 		{
-			IMeshPaintGeometryAdapter* MeshAdapter = FindOrAddGeometryAdapter(MeshComponent);
+			IMeshPaintGeometryAdapter* MeshAdapter = FindGeometryAdapter(MeshComponent);
 			if (MeshAdapter)
 			{
 				SetSpecificTextureOverrideForMesh(*MeshAdapter, GetSelectedTexture());
 			}
 		}
-
-		CleanStaleGeometryAdapters(SelectedMeshComponents);
 	}
 
 	if( bDoRestoreRenTargets && FMeshPaintSettings::Get().ResourceType == EMeshPaintResource::Texture )
@@ -4881,6 +4931,20 @@ void FEdModeMeshPaint::Tick(FEditorViewportClient* ViewportClient,float DeltaTim
 	}
 }
 
+bool FEdModeMeshPaint::ProcessEditDelete()
+{
+	RemoveAllGeometryAdapters();
+
+	if (GUnrealEd->CanDeleteSelectedActors(GetWorld(), true, false))
+	{
+		GEditor->edactDeleteSelected(GetWorld());
+	}
+
+	CreateGeometryAdaptersForSelectedComponents();
+
+	return true;
+}
+
 void FEdModeMeshPaint::DuplicateTextureMaterialCombo()
 {
 	UTexture2D* SelectedTexture = GetSelectedTexture();
@@ -4903,7 +4967,7 @@ void FEdModeMeshPaint::DuplicateTextureMaterialCombo()
 			{
 				UMeshComponent* MeshComponent = MeshComponents[0];
 
-				IMeshPaintGeometryAdapter* MeshAdapter = FindOrAddGeometryAdapter(MeshComponent);
+				IMeshPaintGeometryAdapter* MeshAdapter = FindGeometryAdapter(MeshComponent);
 				if (!MeshAdapter)
 				{
 					return;
diff --git a/Engine/Source/Editor/MeshPaint/Private/MeshPaintStaticMeshAdapter.cpp b/Engine/Source/Editor/MeshPaint/Private/MeshPaintStaticMeshAdapter.cpp
--- a/Engine/Source/Editor/MeshPaint/Private/MeshPaintStaticMeshAdapter.cpp
+++ b/Engine/Source/Editor/MeshPaint/Private/MeshPaintStaticMeshAdapter.cpp
@@ -16,12 +16,14 @@ bool FMeshPaintGeometryAdapterForStaticMeshes::Construct(UMeshComponent* InCompo
 	StaticMeshComponent = Cast<UStaticMeshComponent>(InComponent);
 	if (StaticMeshComponent != nullptr)
 	{
+		StaticMeshComponent->OnStaticMeshChanged().AddRaw(this, &FMeshPaintGeometryAdapterForStaticMeshes::OnStaticMeshChanged);
+
 		if (StaticMeshComponent->StaticMesh != nullptr)
 		{
 			ReferencedStaticMesh = StaticMeshComponent->StaticMesh;
 			PaintingMeshLODIndex = InPaintingMeshLODIndex;
 			UVChannelIndex = InUVChannelIndex;
-			StaticMeshComponent->StaticMesh->OnPostMeshBuild().AddRaw(this, &FMeshPaintGeometryAdapterForStaticMeshes::OnPostMeshBuild);
+			ReferencedStaticMesh->OnPostMeshBuild().AddRaw(this, &FMeshPaintGeometryAdapterForStaticMeshes::OnPostMeshBuild);
 			const bool bSuccess = Initialize();
 			return bSuccess;
 		}
@@ -32,17 +34,37 @@ bool FMeshPaintGeometryAdapterForStaticMeshes::Construct(UMeshComponent* InCompo
 
 FMeshPaintGeometryAdapterForStaticMeshes::~FMeshPaintGeometryAdapterForStaticMeshes()
 {
-	if (StaticMeshComponent != nullptr && StaticMeshComponent->StaticMesh != nullptr)
+	if (StaticMeshComponent != nullptr)
 	{
-		StaticMeshComponent->StaticMesh->OnPostMeshBuild().RemoveAll(this);
+		if (ReferencedStaticMesh != nullptr)
+		{
+			ReferencedStaticMesh->OnPostMeshBuild().RemoveAll(this);
+		}
+
+		StaticMeshComponent->OnStaticMeshChanged().RemoveAll(this);
 	}
 }
 
 void FMeshPaintGeometryAdapterForStaticMeshes::OnPostMeshBuild(UStaticMesh* StaticMesh)
 {
+	check(StaticMesh == ReferencedStaticMesh);
 	Initialize();
 }
 
+void FMeshPaintGeometryAdapterForStaticMeshes::OnStaticMeshChanged(UStaticMeshComponent* InStaticMeshComponent)
+{
+	check(StaticMeshComponent == InStaticMeshComponent);
+	OnRemoved();
+	ReferencedStaticMesh->OnPostMeshBuild().RemoveAll(this);
+	ReferencedStaticMesh = InStaticMeshComponent->StaticMesh;
+	if (ReferencedStaticMesh)
+	{
+		ReferencedStaticMesh->OnPostMeshBuild().AddRaw(this, &FMeshPaintGeometryAdapterForStaticMeshes::OnPostMeshBuild);
+		Initialize();
+		OnAdded();
+	}
+}
+
 bool FMeshPaintGeometryAdapterForStaticMeshes::Initialize()
 {
 	check(ReferencedStaticMesh == StaticMeshComponent->StaticMesh);
@@ -130,7 +152,7 @@ void FMeshPaintGeometryAdapterForStaticMeshes::OnAdded()
 	));
 
 	// If this is the first attempt to add a temporary body setup to the mesh, do it
-	if (StaticMeshReferencers.NumReferencers == 0)
+	if (StaticMeshReferencers.Referencers.Num() == 0)
 	{
 		// Remember the old body setup (this will be added as a GC reference so that it doesn't get destroyed)
 		StaticMeshReferencers.RestoreBodySetup = ReferencedStaticMesh->BodySetup;
@@ -154,8 +176,6 @@ void FMeshPaintGeometryAdapterForStaticMeshes::OnAdded()
 		ReferencedStaticMesh->BodySetup = TempBodySetupRaw;
 	}
 
-	StaticMeshReferencers.NumReferencers++;
-
 	ECollisionEnabled::Type CachedCollisionType = StaticMeshComponent->BodyInstance.GetCollisionEnabled();
 	StaticMeshReferencers.Referencers.Emplace(StaticMeshComponent, CachedCollisionType);
 
@@ -172,30 +192,37 @@ void FMeshPaintGeometryAdapterForStaticMeshes::OnAdded()
 void FMeshPaintGeometryAdapterForStaticMeshes::OnRemoved()
 {
 	check(StaticMeshComponent);
-	check(ReferencedStaticMesh);
+	
+	// If the referenced static mesh has been destroyed (and nulled by GC), don't try to do anything more.
+	// It should be in the process of removing all global geometry adapters if it gets here in this situation.
+	if (!ReferencedStaticMesh)
+	{
+		return;
+	}
 
 	// Remove a reference from the static mesh map
 	FStaticMeshReferencers* StaticMeshReferencers = MeshToComponentMap.Find(ReferencedStaticMesh);
 	check(StaticMeshReferencers);
 	check(StaticMeshReferencers->Referencers.Num() > 0);
-	check(StaticMeshReferencers->NumReferencers > 0);
-	StaticMeshReferencers->NumReferencers--;
-
-	// If the last reference was removed, restore the body setup for all the mesh components which referenced it.
-	if (StaticMeshReferencers->NumReferencers == 0)
-	{
-		ReferencedStaticMesh->BodySetup = StaticMeshReferencers->RestoreBodySetup;
 
-		for (const auto& Info : StaticMeshReferencers->Referencers)
+	int32 Index = StaticMeshReferencers->Referencers.IndexOfByPredicate(
+		[=](const FStaticMeshReferencers::FReferencersInfo& Info)
 		{
-			if (Info.StaticMeshComponent)
-			{
-				Info.StaticMeshComponent->BodyInstance.SetCollisionEnabled(Info.CachedCollisionType, false);
-				Info.StaticMeshComponent->RecreatePhysicsState();
-			}
+			return Info.StaticMeshComponent == this->StaticMeshComponent;
 		}
+	);
+	check(Index != INDEX_NONE);
+
+	StaticMeshComponent->BodyInstance.SetCollisionEnabled(StaticMeshReferencers->Referencers[Index].CachedCollisionType, false);
+	StaticMeshComponent->RecreatePhysicsState();
 
-		StaticMeshReferencers->Referencers.Reset();
+	StaticMeshReferencers->Referencers.RemoveAtSwap(Index);
+
+	// If the last reference was removed, restore the body setup for the static mesh
+	if (StaticMeshReferencers->Referencers.Num() == 0)
+	{
+		ReferencedStaticMesh->BodySetup = StaticMeshReferencers->RestoreBodySetup;
+		verify(MeshToComponentMap.Remove(ReferencedStaticMesh) == 1);
 	}
 }
 
@@ -263,8 +290,10 @@ void FMeshPaintGeometryAdapterForStaticMeshes::SetCurrentUVChannelIndex(int32 In
 
 void FMeshPaintGeometryAdapterForStaticMeshes::AddReferencedObjects(FReferenceCollector& Collector)
 {
-	check(StaticMeshComponent);
-	check(ReferencedStaticMesh);
+	if (!ReferencedStaticMesh)
+	{
+		return;
+	}
 
 	FStaticMeshReferencers* StaticMeshReferencers = MeshToComponentMap.Find(ReferencedStaticMesh);
 	check(StaticMeshReferencers);
diff --git a/Engine/Source/Editor/MeshPaint/Private/MeshPaintStaticMeshAdapter.h b/Engine/Source/Editor/MeshPaint/Private/MeshPaintStaticMeshAdapter.h
--- a/Engine/Source/Editor/MeshPaint/Private/MeshPaintStaticMeshAdapter.h
+++ b/Engine/Source/Editor/MeshPaint/Private/MeshPaintStaticMeshAdapter.h
@@ -37,6 +37,7 @@ protected:
 	virtual void BuildOctree();
 
 	void OnPostMeshBuild(UStaticMesh* StaticMesh);
+	void OnStaticMeshChanged(UStaticMeshComponent* StaticMeshComponent);
 
 	/** Triangle for use in Octree for mesh paint optimization */
 	struct FMeshTriangle
@@ -92,8 +93,7 @@ protected:
 	struct FStaticMeshReferencers
 	{
 		FStaticMeshReferencers()
-			: NumReferencers(0)
-			, RestoreBodySetup(nullptr)
+			: RestoreBodySetup(nullptr)
 		{}
 
 		struct FReferencersInfo
@@ -107,7 +107,6 @@ protected:
 			ECollisionEnabled::Type CachedCollisionType;
 		};
 
-		int32 NumReferencers;
 		TArray<FReferencersInfo> Referencers;
 		UBodySetup* RestoreBodySetup;
 	};
diff --git a/Engine/Source/Editor/MeshPaint/Public/MeshPaintEdMode.h b/Engine/Source/Editor/MeshPaint/Public/MeshPaintEdMode.h
--- a/Engine/Source/Editor/MeshPaint/Public/MeshPaintEdMode.h
+++ b/Engine/Source/Editor/MeshPaint/Public/MeshPaintEdMode.h
@@ -398,6 +398,8 @@ public:
 	virtual bool ShouldDrawWidget() const override { return false; }
 	virtual bool UsesTransformWidget() const override { return false; }
 	virtual void Tick( FEditorViewportClient* ViewportClient, float DeltaTime ) override;
+	virtual EEditAction::Type GetActionEditDelete() override { return EEditAction::Process; }
+	virtual bool ProcessEditDelete() override;
 	// End of FEdMode interface
 
 
@@ -699,24 +701,42 @@ private:
 	/** Returns valid MeshComponents in the current selection */
 	TArray<UMeshComponent*> GetSelectedMeshComponents() const;
 
-	/** Finds an existing geometry adapter for the given component, or creates a new one */
-	IMeshPaintGeometryAdapter* FindOrAddGeometryAdapter(UMeshComponent* MeshComponent);
+	/** Create geometry adapters for selected components */
+	void CreateGeometryAdaptersForSelectedComponents();
 
-	/** Removes stale geometry adapters from the cache (those associated with unselected components) */
-	void CleanStaleGeometryAdapters(const TArray<UMeshComponent*>& ValidComponents);
+	/** Create a new geometry adapter */
+	IMeshPaintGeometryAdapter* AddGeometryAdapter(UMeshComponent* MeshComponent);
+
+	/** Finds an existing geometry adapter for the given component */
+	IMeshPaintGeometryAdapter* FindGeometryAdapter(UMeshComponent* MeshComponent);
+
+	/** Remove a geometry adapter */
+	void RemoveGeometryAdapter(UMeshComponent* MeshComponent);
 
 	/** Removes all geometry adapters from the cache */
 	void RemoveAllGeometryAdapters();
 
-	/** Called when an asset is about to be imported */
-	void OnPreImportAsset(UFactory* Factory, UClass* Class, UObject* Object, const FName& Name, const TCHAR* Type);
+	/** Called prior to saving a level */
+	void OnPreSaveWorld(uint32 SaveFlags, UWorld* World);
+
+	/** Called after saving a level */
+	void OnPostSaveWorld(uint32 SaveFlags, UWorld* World, bool bSuccess);
 
-	/** Called when an asset is about to be reimported */
-	void OnPreReimportAsset(UObject* Object);
+	/** Called when an asset has just been imported */
+	void OnPostImportAsset(UFactory* Factory, UObject* Object);
+
+	/** Called when an asset has just been reimported */
+	void OnPostReimportAsset(UObject* Object, bool bSuccess);
+
+	/** Called when an asset is deleted */
+	void OnAssetRemoved(const FAssetData& AssetData);
 
 	/** Called when the user presses a button on their motion controller device */
 	void OnVRAction( class FEditorViewportClient& ViewportClient, class UViewportInteractor* Interactor,
-		const struct FViewportActionKeyInput& Action, bool& bOutIsInputCaptured, bool& bWasHandled );
+	const struct FViewportActionKeyInput& Action, bool& bOutIsInputCaptured, bool& bWasHandled );
+
+	/** Called when rerunning a construction script causes objects to be replaced */
+	void OnObjectsReplaced(const TMap<UObject*, UObject*>& OldToNewInstanceMap);
 
 private:
 
diff --git a/Engine/Source/Editor/Persona/Private/AnimNotifyDetails.cpp b/Engine/Source/Editor/Persona/Private/AnimNotifyDetails.cpp
--- a/Engine/Source/Editor/Persona/Private/AnimNotifyDetails.cpp
+++ b/Engine/Source/Editor/Persona/Private/AnimNotifyDetails.cpp
@@ -20,7 +20,7 @@ TSharedRef<IDetailCustomization> FAnimNotifyDetails::MakeInstance()
 void FAnimNotifyDetails::CustomizeDetails( IDetailLayoutBuilder& DetailBuilder )
 {
 	const UClass* DetailObjectClass = nullptr;
-	const UClass* BaseClass = DetailBuilder.GetDetailsView().GetBaseClass();
+	const UClass* BaseClass = DetailBuilder.GetBaseClass();
 	TArray<TWeakObjectPtr<UObject>> SelectedObjects;
 	TArray<UClass*> NotifyClasses;
 	DetailBuilder.GetObjectsBeingCustomized(SelectedObjects);
diff --git a/Engine/Source/Editor/Persona/Private/AnimationEditorViewportClient.cpp b/Engine/Source/Editor/Persona/Private/AnimationEditorViewportClient.cpp
--- a/Engine/Source/Editor/Persona/Private/AnimationEditorViewportClient.cpp
+++ b/Engine/Source/Editor/Persona/Private/AnimationEditorViewportClient.cpp
@@ -1589,6 +1589,12 @@ void FAnimationViewportClient::SetViewportType(ELevelViewportType InViewportType
 	FocusViewportOnPreviewMesh();
 }
 
+void FAnimationViewportClient::RotateViewportType()
+{
+	FEditorViewportClient::RotateViewportType();
+	FocusViewportOnPreviewMesh();
+}
+
 bool FAnimationViewportClient::InputKey( FViewport* InViewport, int32 ControllerId, FKey Key, EInputEvent Event, float AmountDepressed, bool bGamepad )
 {
 	const int32 HitX = InViewport->GetMouseX();
diff --git a/Engine/Source/Editor/PhAT/Private/SPhATPreviewViewport.cpp b/Engine/Source/Editor/PhAT/Private/SPhATPreviewViewport.cpp
--- a/Engine/Source/Editor/PhAT/Private/SPhATPreviewViewport.cpp
+++ b/Engine/Source/Editor/PhAT/Private/SPhATPreviewViewport.cpp
@@ -25,6 +25,11 @@ void SPhATPreviewViewport::SetViewportType(ELevelViewportType ViewType)
 	EditorViewportClient->SetViewportType(ViewType);
 }
 
+void SPhATPreviewViewport::RotateViewportType()
+{
+	EditorViewportClient->RotateViewportType();
+}
+
 void SPhATPreviewViewport::Construct(const FArguments& InArgs)
 {
 	PhATPtr = InArgs._PhAT;
diff --git a/Engine/Source/Editor/PhAT/Private/SPhATPreviewViewport.h b/Engine/Source/Editor/PhAT/Private/SPhATPreviewViewport.h
--- a/Engine/Source/Editor/PhAT/Private/SPhATPreviewViewport.h
+++ b/Engine/Source/Editor/PhAT/Private/SPhATPreviewViewport.h
@@ -44,6 +44,7 @@ public:
 	TWeakPtr<SDockTab> ParentTab;
 
 	void SetViewportType(ELevelViewportType ViewType);
+	void RotateViewportType();
 
 protected:
 	/** SEditorViewport interface */
diff --git a/Engine/Source/Editor/PropertyEditor/Private/DetailAdvancedDropdownNode.h b/Engine/Source/Editor/PropertyEditor/Private/DetailAdvancedDropdownNode.h
--- a/Engine/Source/Editor/PropertyEditor/Private/DetailAdvancedDropdownNode.h
+++ b/Engine/Source/Editor/PropertyEditor/Private/DetailAdvancedDropdownNode.h
@@ -30,7 +30,7 @@ private:
 	virtual void GetChildren( TArray< TSharedRef<IDetailTreeNode> >& OutChildren )  override {}
 	virtual void OnItemExpansionChanged( bool bIsExpanded ) override {}
 	virtual bool ShouldBeExpanded() const override { return false; }
-	virtual ENodeVisibility::Type GetVisibility() const override { return ENodeVisibility::Visible; }
+	virtual ENodeVisibility GetVisibility() const override { return ENodeVisibility::Visible; }
 	virtual void FilterNode( const FDetailFilter& InFilter ) override {}
 	virtual void Tick( float DeltaTime ) override {}
 	virtual bool ShouldShowOnlyChildren() const override { return false; }
diff --git a/Engine/Source/Editor/PropertyEditor/Private/DetailCategoryBuilderImpl.cpp b/Engine/Source/Editor/PropertyEditor/Private/DetailCategoryBuilderImpl.cpp
--- a/Engine/Source/Editor/PropertyEditor/Private/DetailCategoryBuilderImpl.cpp
+++ b/Engine/Source/Editor/PropertyEditor/Private/DetailCategoryBuilderImpl.cpp
@@ -42,16 +42,16 @@ void FDetailLayout::AddLayoutInternal( const FDetailLayoutCustomization& Layout,
 }
 
 FDetailLayoutCustomization::FDetailLayoutCustomization()
-	: PropertyRow( NULL )
-	, WidgetDecl( NULL )
-	, CustomBuilderRow( NULL )
+	: PropertyRow(nullptr)
+	, WidgetDecl(nullptr)
+	, CustomBuilderRow(nullptr)
 {
 
 }
 
 TSharedPtr<FPropertyNode> FDetailLayoutCustomization::GetPropertyNode() const 
 {
-	return PropertyRow.IsValid() ? PropertyRow->GetPropertyNode() : NULL;
+	return PropertyRow.IsValid() ? PropertyRow->GetPropertyNode() : nullptr;
 }
 
 FDetailWidgetRow FDetailLayoutCustomization::GetWidgetRow() const
@@ -75,20 +75,20 @@ FDetailWidgetRow FDetailLayoutCustomization::GetWidgetRow() const
 }
 
 FDetailCategoryImpl::FDetailCategoryImpl( FName InCategoryName, TSharedRef<FDetailLayoutBuilderImpl> InDetailLayout )
-	: HeaderContentWidget( NULL )
-	, DetailLayoutBuilder( InDetailLayout )
-	, CategoryName( InCategoryName )
-	, SortOrder( 0 )
-	, bRestoreExpansionState( !ContainsOnlyAdvanced() )
-	, bShouldBeInitiallyCollapsed( false )
-	, bUserShowAdvanced( false )
-	, bForceAdvanced( false )
-	, bHasFilterStrings( false )
-	, bHasVisibleDetails( true )
-	, bIsCategoryVisible( true )
-	, bFavoriteCategory( false )
-{
-	const UStruct* BaseStruct = InDetailLayout->GetDetailsView().GetBaseStruct();
+	: HeaderContentWidget(nullptr)
+	, DetailLayoutBuilder(InDetailLayout)
+	, CategoryName(InCategoryName)
+	, SortOrder(0)
+	, bRestoreExpansionState(!ContainsOnlyAdvanced())
+	, bShouldBeInitiallyCollapsed(false)
+	, bUserShowAdvanced(false)
+	, bForceAdvanced(false)
+	, bHasFilterStrings(false)
+	, bHasVisibleDetails(true)
+	, bIsCategoryVisible(true)
+	, bFavoriteCategory(false)
+{
+	const UStruct* BaseStruct = InDetailLayout->GetRootNode()->GetBaseStructure();
 	// Use the base class name if there is one otherwise this is a generic category not specific to a class
 	FName BaseStructName = BaseStruct ? BaseStruct->GetFName() : FName("Generic");
 
@@ -289,8 +289,8 @@ IDetailPropertyRow* FDetailCategoryImpl::AddExternalProperty( const TArray<UObje
 	}
 
 	FPropertyNodeInitParams InitParams;
-	InitParams.ParentNode = NULL;
-	InitParams.Property = NULL;
+	InitParams.ParentNode = nullptr;
+	InitParams.Property = nullptr;
 	InitParams.ArrayOffset = 0;
 	InitParams.ArrayIndex = INDEX_NONE;
 	// we'll generate the children
@@ -324,7 +324,7 @@ IDetailPropertyRow* FDetailCategoryImpl::AddExternalProperty( const TArray<UObje
 		return NewCustomization.PropertyRow.Get();
 	}
 
-	return NULL;
+	return nullptr;
 }
 
 IDetailPropertyRow* FDetailCategoryImpl::AddExternalProperty( TSharedPtr<FStructOnScope> StructData, FName PropertyName, EPropertyLocation::Type Location/* = EPropertyLocation::Default*/ )
@@ -335,8 +335,8 @@ IDetailPropertyRow* FDetailCategoryImpl::AddExternalProperty( TSharedPtr<FStruct
 	RootPropertyNode->SetStructure(StructData);
 
 	FPropertyNodeInitParams InitParams;
-	InitParams.ParentNode = NULL;
-	InitParams.Property = NULL;
+	InitParams.ParentNode = nullptr;
+	InitParams.Property = nullptr;
 	InitParams.ArrayOffset = 0;
 	InitParams.ArrayIndex = INDEX_NONE;
 	InitParams.bAllowChildren = true;
@@ -378,7 +378,7 @@ IDetailPropertyRow* FDetailCategoryImpl::AddExternalProperty( TSharedPtr<FStruct
 		return NewCustomization.PropertyRow.Get();
 	}
 
-	return NULL;
+	return nullptr;
 }
 
 void FDetailCategoryImpl::AddPropertyNode( TSharedRef<FPropertyNode> PropertyNode, FName InstanceName )
@@ -525,7 +525,7 @@ void FDetailCategoryImpl::SetDisplayName( FName InCategoryName, const FText& Loc
 	}
 	else
 	{
-		const UStruct* BaseStruct = GetParentLayoutImpl().GetDetailsView().GetBaseStruct();
+		const UStruct* BaseStruct = GetParentLayoutImpl().GetRootNode()->GetBaseStructure();
 		// Use the base class name if there is one otherwise this is a generic category not specific to a class
 		FName BaseStructName = BaseStruct ? BaseStruct->GetFName() : FName("Generic");
 
@@ -595,7 +595,7 @@ bool FDetailCategoryImpl::ShouldBeExpanded() const
 	}
 }
 
-ENodeVisibility::Type FDetailCategoryImpl::GetVisibility() const
+ENodeVisibility FDetailCategoryImpl::GetVisibility() const
 {
 	return bHasVisibleDetails && bIsCategoryVisible ? ENodeVisibility::Visible : ENodeVisibility::ForcedHidden;
 }
@@ -817,7 +817,6 @@ void FDetailCategoryImpl::FilterNode( const FDetailFilter& InFilter )
 			bHasVisibleDetails = true;
 			RequestItemExpanded( Child, Child->ShouldBeExpanded() );
 		}
-
 	}
 }
 
diff --git a/Engine/Source/Editor/PropertyEditor/Private/DetailCategoryBuilderImpl.h b/Engine/Source/Editor/PropertyEditor/Private/DetailCategoryBuilderImpl.h
--- a/Engine/Source/Editor/PropertyEditor/Private/DetailCategoryBuilderImpl.h
+++ b/Engine/Source/Editor/PropertyEditor/Private/DetailCategoryBuilderImpl.h
@@ -134,8 +134,8 @@ public:
 	virtual IDetailCategoryBuilder& OnExpansionChanged( FOnBooleanValueChanged InOnExpansionChanged ) override;
 	virtual IDetailCategoryBuilder& RestoreExpansionState( bool bRestore ) override;
 	virtual IDetailCategoryBuilder& HeaderContent( TSharedRef<SWidget> InHeaderContent ) override;
-	virtual IDetailPropertyRow& AddProperty( FName PropertyPath, UClass* ClassOuter = NULL, FName InstanceName = NAME_None,  EPropertyLocation::Type Location = EPropertyLocation::Default ) override;
-	virtual IDetailPropertyRow& AddProperty( TSharedPtr<IPropertyHandle> PropertyHandle,  EPropertyLocation::Type Location = EPropertyLocation::Default ) override;
+	virtual IDetailPropertyRow& AddProperty( FName PropertyPath, UClass* ClassOuter = nullptr, FName InstanceName = NAME_None,  EPropertyLocation::Type Location = EPropertyLocation::Default ) override;
+	virtual IDetailPropertyRow& AddProperty( TSharedPtr<IPropertyHandle> PropertyHandle, EPropertyLocation::Type Location = EPropertyLocation::Default ) override;
 	virtual IDetailPropertyRow* AddExternalProperty( const TArray<UObject*>& Objects, FName PropertyName, EPropertyLocation::Type Location = EPropertyLocation::Default ) override;
 	virtual IDetailPropertyRow* AddExternalProperty( TSharedPtr<FStructOnScope> StructData, FName PropertyName, EPropertyLocation::Type Location = EPropertyLocation::Default ) override;
 	virtual IDetailLayoutBuilder& GetParentLayout() const override { return *DetailLayoutBuilder.Pin(); }
@@ -151,7 +151,7 @@ public:
 	virtual TSharedRef< ITableRow > GenerateNodeWidget( const TSharedRef<STableViewBase>& OwnerTable, const FDetailColumnSizeData& ColumnSizeData, const TSharedRef<IPropertyUtilities>& PropertyUtilities, bool bAllowFavoriteSystem) override;
 	virtual void GetChildren( TArray< TSharedRef<IDetailTreeNode> >& OutChildren ) override;
 	virtual bool ShouldBeExpanded() const override;
-	virtual ENodeVisibility::Type GetVisibility() const override;
+	virtual ENodeVisibility GetVisibility() const override;
 	virtual void FilterNode( const FDetailFilter& DetailFilter ) override;
 	virtual void Tick( float DeltaTime ) override {}
 	virtual bool ShouldShowOnlyChildren() const override { return false; }
@@ -381,7 +381,6 @@ private:
 	bool bHasVisibleDetails:1;
 	/** true if the category is visible at all */
 	bool bIsCategoryVisible:1;
-
 	/*true if the category is the special favorite category, all property in the layout will be display when we generate the roottree */
 	bool bFavoriteCategory:1;
 };
diff --git a/Engine/Source/Editor/PropertyEditor/Private/DetailCategoryGroupNode.h b/Engine/Source/Editor/PropertyEditor/Private/DetailCategoryGroupNode.h
--- a/Engine/Source/Editor/PropertyEditor/Private/DetailCategoryGroupNode.h
+++ b/Engine/Source/Editor/PropertyEditor/Private/DetailCategoryGroupNode.h
@@ -36,7 +36,7 @@ private:
 	virtual IDetailsViewPrivate& GetDetailsView() const override{ return ParentCategory.GetDetailsView(); }
 	virtual void OnItemExpansionChanged( bool bIsExpanded ) override {}
 	virtual bool ShouldBeExpanded() const override { return true; }
-	virtual ENodeVisibility::Type GetVisibility() const override { return bShouldBeVisible ? ENodeVisibility::Visible : ENodeVisibility::HiddenDueToFiltering; }
+	virtual ENodeVisibility GetVisibility() const override { return bShouldBeVisible ? ENodeVisibility::Visible : ENodeVisibility::HiddenDueToFiltering; }
 	virtual TSharedRef< ITableRow > GenerateNodeWidget( const TSharedRef<STableViewBase>& OwnerTable, const FDetailColumnSizeData& ColumnSizeData, const TSharedRef<IPropertyUtilities>& PropertyUtilities, bool bAllowFavoriteSystem) override;
 	virtual void GetChildren( TArray< TSharedRef<IDetailTreeNode> >& OutChildren )  override;
 	virtual void FilterNode( const FDetailFilter& InFilter ) override;
diff --git a/Engine/Source/Editor/PropertyEditor/Private/DetailItemNode.cpp b/Engine/Source/Editor/PropertyEditor/Private/DetailItemNode.cpp
--- a/Engine/Source/Editor/PropertyEditor/Private/DetailItemNode.cpp
+++ b/Engine/Source/Editor/PropertyEditor/Private/DetailItemNode.cpp
@@ -189,7 +189,7 @@ void FDetailItemNode::GetChildren( FDetailNodeList& OutChildren )
 	{
 		TSharedRef<IDetailTreeNode>& Child = Children[ChildIndex];
 
-		ENodeVisibility::Type ChildVisibility = Child->GetVisibility();
+		ENodeVisibility ChildVisibility = Child->GetVisibility();
 
 		// Report the child if the child is visible or we are visible due to filtering and there were no filtered children.  
 		// If we are visible due to filtering and so is a child, we only show that child.  
@@ -264,7 +264,7 @@ bool FDetailItemNode::ShouldBeExpanded() const
 	return bShouldBeExpanded;
 }
 
-ENodeVisibility::Type FDetailItemNode::GetVisibility() const
+ENodeVisibility FDetailItemNode::GetVisibility() const
 {
 	const bool bHasAnythingToShow = Customization.IsValidCustomization();
 
@@ -273,7 +273,7 @@ ENodeVisibility::Type FDetailItemNode::GetVisibility() const
 		|| (Customization.HasCustomWidget() && Customization.WidgetDecl->VisibilityAttr.Get() != EVisibility::Visible )
 		|| (Customization.HasPropertyNode() && Customization.PropertyRow->GetPropertyVisibility() != EVisibility::Visible );
 
-	ENodeVisibility::Type Visibility;
+	ENodeVisibility Visibility;
 	if( bIsForcedHidden )
 	{
 		Visibility = ENodeVisibility::ForcedHidden;
@@ -329,7 +329,7 @@ static bool PassesAllFilters( const FDetailLayoutCustomization& InCustomization,
 
 			const bool bPassesSearchFilter = bSearchFilterIsEmpty || ( bIsNotBeingFiltered || bIsSeenDueToFiltering || bIsParentSeenDueToFiltering );
 			const bool bPassesModifiedFilter = bPassesSearchFilter && ( InFilter.bShowOnlyModifiedProperties == false || PropertyNodePin->GetDiffersFromDefault() == true );
-			const bool bPassesDifferingFilter = InFilter.bShowOnlyDiffering ? InFilter.WhitelistedProperties.Find(*FPropertyNode::CreatePropertyPath(PropertyNodePin.ToSharedRef())) != NULL : true;
+			const bool bPassesDifferingFilter = InFilter.bShowOnlyDiffering ? InFilter.WhitelistedProperties.Find(*FPropertyNode::CreatePropertyPath(PropertyNodePin.ToSharedRef())) != nullptr : true;
 
 			// The property node is visible (note categories are never visible unless they have a child that is visible )
 			bPassesAllFilters = (bPassesSearchFilter && bPassesModifiedFilter && bPassesDifferingFilter) || bPassesCategoryFilter;
diff --git a/Engine/Source/Editor/PropertyEditor/Private/DetailItemNode.h b/Engine/Source/Editor/PropertyEditor/Private/DetailItemNode.h
--- a/Engine/Source/Editor/PropertyEditor/Private/DetailItemNode.h
+++ b/Engine/Source/Editor/PropertyEditor/Private/DetailItemNode.h
@@ -43,7 +43,7 @@ public:
 	virtual void GetChildren( FDetailNodeList& OutChildren )  override;
 	virtual void OnItemExpansionChanged( bool bInIsExpanded ) override;
 	virtual bool ShouldBeExpanded() const override;
-	virtual ENodeVisibility::Type GetVisibility() const override;
+	virtual ENodeVisibility GetVisibility() const override;
 	virtual void FilterNode( const FDetailFilter& InFilter ) override;
 	virtual void Tick( float DeltaTime ) override;
 	virtual bool ShouldShowOnlyChildren() const override;
diff --git a/Engine/Source/Editor/PropertyEditor/Private/DetailLayoutBuilderImpl.cpp b/Engine/Source/Editor/PropertyEditor/Private/DetailLayoutBuilderImpl.cpp
--- a/Engine/Source/Editor/PropertyEditor/Private/DetailLayoutBuilderImpl.cpp
+++ b/Engine/Source/Editor/PropertyEditor/Private/DetailLayoutBuilderImpl.cpp
@@ -5,13 +5,15 @@
 #include "PropertyHandleImpl.h"
 #include "PropertyEditorHelpers.h"
 #include "IPropertyUtilities.h"
+#include "DetailMultiTopLevelObjectRootNode.h"
+#include "IDetailRootObjectCustomization.h"
 
-
-FDetailLayoutBuilderImpl::FDetailLayoutBuilderImpl(FClassToPropertyMap& InPropertyMap, const TSharedRef< class IPropertyUtilities >& InPropertyUtilities, const TSharedRef< IDetailsViewPrivate >& InDetailsView)
-	: PropertyMap( InPropertyMap )
+FDetailLayoutBuilderImpl::FDetailLayoutBuilderImpl(TSharedPtr<FComplexPropertyNode>& InRootNode, FClassToPropertyMap& InPropertyMap, const TSharedRef< class IPropertyUtilities >& InPropertyUtilities, const TSharedRef< IDetailsViewPrivate >& InDetailsView)
+	: RootNode( InRootNode )
+	, PropertyMap( InPropertyMap )
 	, PropertyDetailsUtilities( InPropertyUtilities )
 	, DetailsView( *InDetailsView )
-	, CurrentCustomizationClass( NULL )
+	, CurrentCustomizationClass( nullptr )
 {
 
 }
@@ -74,6 +76,15 @@ TSharedRef<IPropertyHandle> FDetailLayoutBuilderImpl::GetProperty( const FName P
 
 }
 
+FName FDetailLayoutBuilderImpl::GetTopLevelProperty()
+{
+	for (auto It = PropertyMap.CreateConstIterator(); It; ++It)
+	{
+		return It.Key();
+	}
+	return NAME_None;
+}
+
 void FDetailLayoutBuilderImpl::HideProperty( const TSharedPtr<IPropertyHandle> PropertyHandle )
 {
 	if( PropertyHandle.IsValid() && PropertyHandle->IsValidHandle() )
@@ -132,7 +143,9 @@ void FDetailLayoutBuilderImpl::BuildCategories( const FCategoryMap& CategoryMap,
 		//If there is a delimiter in the name it mean its a sub category, we dont show sub category at the root level
 		FString CategoryDelimiterString;
 		CategoryDelimiterString.AppendChar(FPropertyNodeConstants::CategoryDelimiterChar);
-		const bool bCategoryHidden = GetDetailsView().IsCategoryHiddenByClass(DetailCategory->GetCategoryName()) || ForceHiddenCategories.Contains(DetailCategory->GetCategoryName()) || DetailCategory->GetCategoryName().ToString().Contains(CategoryDelimiterString);
+		TSharedPtr<FComplexPropertyNode> RootPropertyNode = GetRootNode();
+		const bool bCategoryHidden = GetDetailsView().IsCategoryHiddenByClass(RootPropertyNode, DetailCategory->GetCategoryName()) 
+			|| ForceHiddenCategories.Contains(DetailCategory->GetCategoryName()) || DetailCategory->GetCategoryName().ToString().Contains(CategoryDelimiterString);
 
 		if( !bCategoryHidden )
 		{
@@ -176,37 +189,52 @@ void FDetailLayoutBuilderImpl::GenerateDetailLayout()
 	SimpleCategories.Sort( FCompareFDetailCategoryImpl() );
 	AdvancedOnlyCategories.Sort( FCompareFDetailCategoryImpl() );
 
-	/** Merge the two category lists in sorted order */
-	for( int32 CategoryIndex = 0; CategoryIndex < SimpleCategories.Num(); ++CategoryIndex )
+	FDetailNodeList CategoryNodes;
+
+	// Merge the two category lists in sorted order
+	for (int32 CategoryIndex = 0; CategoryIndex < SimpleCategories.Num(); ++CategoryIndex)
 	{
-		AllRootTreeNodes.Add( SimpleCategories[CategoryIndex] );
+		CategoryNodes.Add(SimpleCategories[CategoryIndex]);
 	}
 
-	for( int32 CategoryIndex = 0; CategoryIndex < AdvancedOnlyCategories.Num(); ++CategoryIndex )
+	for (int32 CategoryIndex = 0; CategoryIndex < AdvancedOnlyCategories.Num(); ++CategoryIndex)
 	{
-		AllRootTreeNodes.Add( AdvancedOnlyCategories[CategoryIndex] );
+		CategoryNodes.Add(AdvancedOnlyCategories[CategoryIndex]);
 	}
 
+	if(DetailsView.ContainsMultipleTopLevelObjects())
+	{
+		// This should always exist here
+		UObject* RootObject = RootNode.Pin()->AsObjectNode()->GetUObject(0);
+		check(RootObject);
 
+		TSharedPtr<IDetailRootObjectCustomization> RootObjectCustomization = DetailsView.GetRootObjectCustomization();
+
+		// there are multiple objects in the details panel.  Separate each one with a unique object name node to differentiate them
+		AllRootTreeNodes.Add( MakeShareable( new FDetailMultiTopLevelObjectRootNode( CategoryNodes, RootObjectCustomization, DetailsView, *RootObject) ) );
+	}
+	else
+	{
+		// The categories are the root in this case
+		AllRootTreeNodes = CategoryNodes;
+	}
 }
 
 void FDetailLayoutBuilderImpl::FilterDetailLayout( const FDetailFilter& InFilter )
 {
 	CurrentFilter = InFilter;
-	FilteredRootTreeNodes.Empty();
+	FilteredRootTreeNodes.Reset();
 
-	for( int32 RootNodeIndex = 0; RootNodeIndex < AllRootTreeNodes.Num(); ++RootNodeIndex )
+	for(int32 RootNodeIndex = 0; RootNodeIndex < AllRootTreeNodes.Num(); ++RootNodeIndex)
 	{
-		TSharedRef<IDetailTreeNode>& RootNode = AllRootTreeNodes[RootNodeIndex];
+		TSharedRef<IDetailTreeNode>& RootTreeNode = AllRootTreeNodes[RootNodeIndex];
 
-		// No parent
-		const bool bParentVisibleDueToFiltering = false;
-		RootNode->FilterNode( InFilter );
+		RootTreeNode->FilterNode(InFilter);
 
-		if( RootNode->GetVisibility() == ENodeVisibility::Visible )
+		if(RootTreeNode->GetVisibility() == ENodeVisibility::Visible)
 		{
-			FilteredRootTreeNodes.Add( RootNode );
-			DetailsView.RequestItemExpanded( RootNode, RootNode->ShouldBeExpanded() );
+			FilteredRootTreeNodes.Add(RootTreeNode);
+			DetailsView.RequestItemExpanded(RootTreeNode, RootTreeNode->ShouldBeExpanded());
 		}
 	}
 }
@@ -254,7 +282,7 @@ static void GetPropertyAndIndex( const FString& PathNode, FString& OutProperty,
  */
 static TSharedPtr<FPropertyNode> FindChildPropertyNode( FPropertyNode& InParentNode, const FString& PropertyName, int32 Index )
 {
-	TSharedPtr<FPropertyNode> FoundNode(NULL);
+	TSharedPtr<FPropertyNode> FoundNode(nullptr);
 
 	// search each child for a property with the provided name
 	for( int32 ChildIndex = 0; ChildIndex < InParentNode.GetNumChildNodes(); ++ChildIndex )
@@ -280,7 +308,7 @@ static TSharedPtr<FPropertyNode> FindChildPropertyNode( FPropertyNode& InParentN
 
 TSharedPtr<FPropertyNode> FDetailLayoutBuilderImpl::GetPropertyNode( TSharedPtr<IPropertyHandle> PropertyHandle ) const
 {
-	TSharedPtr<FPropertyNode> PropertyNode = NULL;
+	TSharedPtr<FPropertyNode> PropertyNode = nullptr;
 	if( PropertyHandle->IsValidHandle() )
 	{
 		PropertyNode = StaticCastSharedPtr<FPropertyHandleBase>(PropertyHandle)->GetPropertyNode();
@@ -310,7 +338,7 @@ TSharedPtr<FPropertyNode> FDetailLayoutBuilderImpl::GetPropertyNodeInternal( con
 	}
 	// The class to find properties in defaults to the class currently being customized
 	FName ClassName = CurrentCustomizationClass ? CurrentCustomizationClass->GetFName() : NAME_None;
-	if( ClassOutermost != NULL )
+	if( ClassOutermost != nullptr)
 	{
 		// The requested a different class
 		ClassName = ClassOutermost->GetFName();
@@ -386,7 +414,7 @@ TSharedPtr<FPropertyNode> FDetailLayoutBuilderImpl::GetPropertyNodeInternal( con
 		}
 	}
 
-	return NULL;
+	return nullptr;
 }
 
 
@@ -403,7 +431,7 @@ TSharedRef<IPropertyHandle> FDetailLayoutBuilderImpl::GetPropertyHandle( TShared
 	else
 	{
 		// Invalid handle
-		PropertyHandle = MakeShareable( new FPropertyHandleBase( NULL, NULL, NULL) ); 
+		PropertyHandle = MakeShareable( new FPropertyHandleBase(nullptr, nullptr, nullptr) );
 	}	
 
 	return PropertyHandle.ToSharedRef();
@@ -450,38 +478,14 @@ void FDetailLayoutBuilderImpl::GetObjectsBeingCustomized( TArray< TWeakObjectPtr
 {
 	OutObjects.Empty();
 
-	// The class to find properties in defaults to the class currently being customized
-	FName ClassName = CurrentCustomizationClass ? CurrentCustomizationClass->GetFName() : NAME_None;
-	
-	if( ClassName != NAME_None && CurrentCustomizationVariableName != NAME_None )
+	FObjectPropertyNode* RootObjectNode = RootNode.IsValid() ? RootNode.Pin()->AsObjectNode() : nullptr;
+	if(RootObjectNode)
 	{
-		// If this fails there are no properties associated with the class name provided
-		FClassInstanceToPropertyMap* ClassInstanceToPropertyMapPtr = PropertyMap.Find( ClassName );
-
-		if( ClassInstanceToPropertyMapPtr )
+		for (int32 ObjectIndex = 0; ObjectIndex < RootObjectNode->GetNumObjects(); ++ObjectIndex)
 		{
-			FClassInstanceToPropertyMap& ClassInstanceToPropertyMap = *ClassInstanceToPropertyMapPtr;
-		
-			FPropertyNodeMap* PropertyNodeMapPtr = ClassInstanceToPropertyMap.Find( CurrentCustomizationVariableName );
-
-			if( PropertyNodeMapPtr )
-			{
-				FPropertyNodeMap& PropertyNodeMap = *PropertyNodeMapPtr;
-				FObjectPropertyNode* ParentObjectProperty = PropertyNodeMap.ParentProperty ? PropertyNodeMap.ParentProperty->AsObjectNode() : NULL;
-				if (ParentObjectProperty)
-				{
-					for (int32 ObjectIndex = 0; ObjectIndex < ParentObjectProperty->GetNumObjects(); ++ObjectIndex)
-					{
-						OutObjects.Add(ParentObjectProperty->GetUObject(ObjectIndex));
-					}
-				}
-			}
+			OutObjects.Add(RootObjectNode->GetUObject(ObjectIndex));
 		}
 	}
-	else
-	{
-		OutObjects = DetailsView.GetSelectedObjects();
-	}
 }
 
 const TSharedRef< IPropertyUtilities >& FDetailLayoutBuilderImpl::GetPropertyUtilities() const
@@ -489,6 +493,10 @@ const TSharedRef< IPropertyUtilities >& FDetailLayoutBuilderImpl::GetPropertyUti
 	return PropertyDetailsUtilities;
 }
 
+UClass* FDetailLayoutBuilderImpl::GetBaseClass() const
+{
+	return RootNode.IsValid() ? Cast<UClass>(RootNode.Pin()->GetBaseStructure()) : nullptr;
+}
 void FDetailLayoutBuilderImpl::SetCustomProperty( const TSharedPtr<FPropertyNode>& PropertyNode ) 
 {
 	PropertyNode->SetNodeFlags( EPropertyNodeFlags::IsCustomized, true );
diff --git a/Engine/Source/Editor/PropertyEditor/Private/DetailLayoutBuilderImpl.h b/Engine/Source/Editor/PropertyEditor/Private/DetailLayoutBuilderImpl.h
--- a/Engine/Source/Editor/PropertyEditor/Private/DetailLayoutBuilderImpl.h
+++ b/Engine/Source/Editor/PropertyEditor/Private/DetailLayoutBuilderImpl.h
@@ -6,13 +6,18 @@ class FDetailCategoryImpl;
 class FDetailLayoutBuilderImpl : public IDetailLayoutBuilder, public TSharedFromThis<FDetailLayoutBuilderImpl>
 {
 public:
-	FDetailLayoutBuilderImpl(FClassToPropertyMap& InPropertyMap, const TSharedRef< class IPropertyUtilities >& InPropertyUtilities, const TSharedRef< IDetailsViewPrivate >& InDetailsView);
+	FDetailLayoutBuilderImpl(
+		TSharedPtr<FComplexPropertyNode>& InRootNode, 
+		FClassToPropertyMap& InPropertyMap, 
+		const TSharedRef< class IPropertyUtilities >& InPropertyUtilities, 
+		const TSharedRef< IDetailsViewPrivate >& InDetailsView);
 
 	/** IDetailLayoutBuilder Interface */
 	virtual const IDetailsView& GetDetailsView() const override;
 	virtual void GetObjectsBeingCustomized( TArray< TWeakObjectPtr<UObject> >& OutObjects ) const override;
 	virtual IDetailCategoryBuilder& EditCategory( FName CategoryName, const FText& NewLocalizedDisplayName = FText::GetEmpty(), ECategoryPriority::Type CategoryType = ECategoryPriority::Default ) override;
 	virtual TSharedRef<IPropertyHandle> GetProperty( const FName PropertyPath, const UClass* ClassOutermost, FName InInstanceName ) override;
+	virtual FName GetTopLevelProperty() override;
 	virtual void HideProperty( const TSharedPtr<IPropertyHandle> Property ) override;
 	virtual void HideProperty( FName PropertyPath, const UClass* ClassOutermost = NULL, FName InstanceName = NAME_None ) override;
 	virtual void ForceRefreshDetails() override;
@@ -21,6 +26,7 @@ public:
 	virtual bool IsPropertyVisible( const struct FPropertyAndParent& PropertyAndParent ) const override;
 	virtual void HideCategory( FName CategoryName ) override;
 	virtual const TSharedRef< IPropertyUtilities >& GetPropertyUtilities() const override; 
+	virtual UClass* GetBaseClass() const override;
 
 	/**
 	 * Creates a default category. The SDetails view will generate widgets in default categories 
@@ -35,7 +41,7 @@ public:
 	bool HasCategory(FName CategoryName);
 
 	/**
-	 * Generates the layout for this detail builder                                                              
+	 * Generates the layout for this detail builder
 	 */
 	void GenerateDetailLayout();
 
@@ -84,7 +90,7 @@ public:
 	/**
 	 * @return All tree nodes that should be visible in the tree                                                              
 	 */
-	TArray< TSharedRef<IDetailTreeNode> >& GetRootTreeNodes() { return FilteredRootTreeNodes; }
+	FDetailNodeList& GetRootTreeNodes() { return FilteredRootTreeNodes; }
 
 	/**
 	 * @return true if the layout has any details
@@ -147,6 +153,9 @@ public:
 
 	/** @return The details view that owns this layout */
 	IDetailsViewPrivate& GetDetailsView() { return DetailsView; }
+
+	/** @return The root node for this customization */
+	TSharedPtr<FComplexPropertyNode> GetRootNode() const { return RootNode.Pin(); }
 private:
 	/**
 	 * Finds a property node for the current property by searching in a fast lookup map or a path search if required
@@ -164,6 +173,8 @@ private:
 	void BuildCategories( const FCategoryMap& CategoryMap, TArray< TSharedRef<FDetailCategoryImpl> >& OutSimpleCategories, TArray< TSharedRef<FDetailCategoryImpl> >& OutAdvancedCategories );
 
 private:
+	/** The root property node for this customization */
+	TWeakPtr<FComplexPropertyNode> RootNode;
 	/** A mapping of category names to categories which have been customized */
 	FCategoryMap CustomCategoryMap;
 	/** A mapping of category names to categories which have been not customized */
@@ -189,3 +200,4 @@ private:
 	/** The current class being customized */
 	UStruct* CurrentCustomizationClass;
 };
+
diff --git a/Engine/Source/Editor/PropertyEditor/Private/DetailMultiTopLevelObjectRootNode.cpp b/Engine/Source/Editor/PropertyEditor/Private/DetailMultiTopLevelObjectRootNode.cpp
new file mode 100644
--- /dev/null
+++ b/Engine/Source/Editor/PropertyEditor/Private/DetailMultiTopLevelObjectRootNode.cpp
@@ -0,0 +1,160 @@
+// Copyright 1998-2016 Epic Games, Inc. All Rights Reserved.
+
+#include "PropertyEditorPrivatePCH.h"
+#include "DetailMultiTopLevelObjectRootNode.h"
+#include "IDetailRootObjectCustomization.h"
+
+void SDetailMultiTopLevelObjectTableRow::Construct( const FArguments& InArgs, TSharedRef<IDetailTreeNode> InOwnerTreeNode, const TSharedRef<SWidget>& InCustomizedWidgetContents, const TSharedRef<STableViewBase>& InOwnerTableView )
+{
+	OwnerTreeNode = InOwnerTreeNode;
+	bShowExpansionArrow = InArgs._ShowExpansionArrow;
+
+	ChildSlot
+	[	
+		SNew( SBox )
+		.Padding( FMargin( 0.0f, 0.0f, SDetailTableRowBase::ScrollbarPaddingSize, 0.0f ) )
+		[
+			SNew( SHorizontalBox )
+			+SHorizontalBox::Slot()
+			.VAlign(VAlign_Center)
+			.Padding(2.0f, 2.0f, 2.0f, 2.0f)
+			.AutoWidth()
+			[
+				SNew( SExpanderArrow, SharedThis(this) )
+				.Visibility(bShowExpansionArrow ? EVisibility::Visible : EVisibility::Collapsed)
+			]
+			+SHorizontalBox::Slot()
+			[
+				InCustomizedWidgetContents
+			]
+		]
+	];
+
+	STableRow< TSharedPtr< IDetailTreeNode > >::ConstructInternal(
+		STableRow::FArguments()
+			.Style(FEditorStyle::Get(), "DetailsView.TreeView.TableRow")
+			.ShowSelection(false),
+		InOwnerTableView
+	);
+}
+
+
+const FSlateBrush* SDetailMultiTopLevelObjectTableRow::GetBackgroundImage() const
+{
+	if (IsHovered())
+	{
+		return IsItemExpanded() ? FEditorStyle::GetBrush("DetailsView.CategoryTop_Hovered") : FEditorStyle::GetBrush("DetailsView.CollapsedCategory_Hovered");
+	}
+	else
+	{
+		return IsItemExpanded() ? FEditorStyle::GetBrush("DetailsView.CategoryTop") : FEditorStyle::GetBrush("DetailsView.CollapsedCategory");
+	}
+}
+
+FReply SDetailMultiTopLevelObjectTableRow::OnMouseButtonDown( const FGeometry& MyGeometry, const FPointerEvent& MouseEvent )
+{
+	if (bShowExpansionArrow && MouseEvent.GetEffectingButton() == EKeys::LeftMouseButton)
+	{
+		ToggleExpansion();
+		return FReply::Handled();
+	}
+	else
+	{
+		return FReply::Unhandled();
+	}
+}
+
+FReply SDetailMultiTopLevelObjectTableRow::OnMouseButtonDoubleClick( const FGeometry& InMyGeometry, const FPointerEvent& InMouseEvent )
+{
+	return OnMouseButtonDown(InMyGeometry, InMouseEvent);
+}
+
+
+FDetailMultiTopLevelObjectRootNode::FDetailMultiTopLevelObjectRootNode( const FDetailNodeList& InChildNodes, const TSharedPtr<IDetailRootObjectCustomization>& InRootObjectCustomization, IDetailsViewPrivate& InDetailsView, const UObject& InRootObject )
+	: ChildNodes(InChildNodes)
+	, DetailsView(InDetailsView)
+	, RootObjectCustomization(InRootObjectCustomization)
+	, RootObject(&InRootObject)
+	, NodeName(InRootObject.GetFName())
+	, bShouldBeVisible(false)
+{
+}
+
+ENodeVisibility FDetailMultiTopLevelObjectRootNode::GetVisibility() const
+{
+	ENodeVisibility FinalVisibility = ENodeVisibility::Visible;
+	if(RootObjectCustomization.IsValid() && RootObject.IsValid() && !RootObjectCustomization.Pin()->IsObjectVisible(RootObject.Get()))
+	{
+		FinalVisibility = ENodeVisibility::ForcedHidden;
+	}
+	else
+	{
+		FinalVisibility = bShouldBeVisible ? ENodeVisibility::Visible : ENodeVisibility::HiddenDueToFiltering;
+	}
+
+	return FinalVisibility;
+}
+
+TSharedRef< ITableRow > FDetailMultiTopLevelObjectRootNode::GenerateNodeWidget(const TSharedRef<STableViewBase>& OwnerTable, const FDetailColumnSizeData& ColumnSizeData, const TSharedRef<IPropertyUtilities>& PropertyUtilities, bool bAllowFavoriteSystem)
+{
+	TSharedPtr<SWidget> HeaderWidget;
+	if(RootObjectCustomization.IsValid() && RootObject.IsValid())
+	{
+		HeaderWidget = RootObjectCustomization.Pin()->CustomizeObjectHeader(RootObject.Get());
+	}
+
+	if(!HeaderWidget.IsValid())
+	{
+		// no customization was supplied or was passed back from the interface as invalid
+		// just make a text block with the name
+		HeaderWidget =
+			SNew(STextBlock)
+			.Font(FEditorStyle::GetFontStyle("DetailsView.CategoryFontStyle"))
+			.Text(FText::FromName(NodeName));
+	}
+
+	return SNew(SDetailMultiTopLevelObjectTableRow, AsShared(), HeaderWidget.ToSharedRef(), OwnerTable);
+}
+
+
+void FDetailMultiTopLevelObjectRootNode::GetChildren( TArray< TSharedRef<IDetailTreeNode> >& OutChildren )
+{
+	for( int32 ChildIndex = 0; ChildIndex < ChildNodes.Num(); ++ChildIndex )
+	{
+		TSharedRef<IDetailTreeNode>& Child = ChildNodes[ChildIndex];
+		if( Child->GetVisibility() == ENodeVisibility::Visible )
+		{
+			if( Child->ShouldShowOnlyChildren() )
+			{
+				Child->GetChildren( OutChildren );
+			}
+			else
+			{
+				OutChildren.Add( Child );
+			}
+		}
+	}
+}
+
+void FDetailMultiTopLevelObjectRootNode::FilterNode( const FDetailFilter& InFilter )
+{
+	bShouldBeVisible = false;
+	for( int32 ChildIndex = 0; ChildIndex < ChildNodes.Num(); ++ChildIndex )
+	{
+		TSharedRef<IDetailTreeNode>& Child = ChildNodes[ChildIndex];
+
+		Child->FilterNode( InFilter );
+
+		if( Child->GetVisibility() == ENodeVisibility::Visible )
+		{
+			bShouldBeVisible = true;
+
+			DetailsView.RequestItemExpanded( Child, Child->ShouldBeExpanded() );
+		}
+	}
+}
+
+bool FDetailMultiTopLevelObjectRootNode::ShouldShowOnlyChildren() const
+{
+	return RootObjectCustomization.IsValid() && RootObject.IsValid() ? !RootObjectCustomization.Pin()->ShouldDisplayHeader(RootObject.Get()) : false;
+}
diff --git a/Engine/Source/Editor/PropertyEditor/Private/DetailMultiTopLevelObjectRootNode.h b/Engine/Source/Editor/PropertyEditor/Private/DetailMultiTopLevelObjectRootNode.h
new file mode 100644
--- /dev/null
+++ b/Engine/Source/Editor/PropertyEditor/Private/DetailMultiTopLevelObjectRootNode.h
@@ -0,0 +1,52 @@
+// Copyright 1998-2016 Epic Games, Inc. All Rights Reserved.
+
+#pragma once
+
+#include "SDetailTableRowBase.h"
+
+class SDetailMultiTopLevelObjectTableRow : public SDetailTableRowBase
+{
+public:
+	SLATE_BEGIN_ARGS(SDetailMultiTopLevelObjectTableRow)
+		: _DisplayName()
+		, _ShowExpansionArrow(false)
+	{}
+		SLATE_ARGUMENT( FText, DisplayName )
+		SLATE_ARGUMENT( bool, ShowExpansionArrow )
+	SLATE_END_ARGS()
+
+	void Construct( const FArguments& InArgs, TSharedRef<IDetailTreeNode> InOwnerTreeNode, const TSharedRef<SWidget>& InCustomizedWidgetContents, const TSharedRef<STableViewBase>& InOwnerTableView );
+private:
+	const FSlateBrush* GetBackgroundImage() const;
+private:
+	virtual FReply OnMouseButtonDown( const FGeometry& MyGeometry, const FPointerEvent& MouseEvent ) override;
+	virtual FReply OnMouseButtonDoubleClick( const FGeometry& InMyGeometry, const FPointerEvent& InMouseEvent ) override;
+private:
+	bool bShowExpansionArrow;
+};
+
+
+class FDetailMultiTopLevelObjectRootNode : public IDetailTreeNode, public TSharedFromThis<FDetailMultiTopLevelObjectRootNode>
+{
+public:
+	FDetailMultiTopLevelObjectRootNode( const FDetailNodeList& InChildNodes, const TSharedPtr<IDetailRootObjectCustomization>& RootObjectCustomization, IDetailsViewPrivate& InDetailsView, const UObject& InRootObject );
+
+private:
+	virtual IDetailsViewPrivate& GetDetailsView() const override{ return DetailsView; }
+	virtual void OnItemExpansionChanged( bool bIsExpanded ) override {}
+	virtual bool ShouldBeExpanded() const override { return true; }
+	virtual ENodeVisibility GetVisibility() const override;
+	virtual TSharedRef< ITableRow > GenerateNodeWidget( const TSharedRef<STableViewBase>& OwnerTable, const FDetailColumnSizeData& ColumnSizeData, const TSharedRef<IPropertyUtilities>& PropertyUtilities, bool bAllowFavoriteSystem) override;
+	virtual void GetChildren( TArray< TSharedRef<IDetailTreeNode> >& OutChildren )  override;
+	virtual void FilterNode( const FDetailFilter& InFilter ) override;
+	virtual void Tick( float DeltaTime ) override {}
+	virtual bool ShouldShowOnlyChildren() const override;
+	virtual FName GetNodeName() const override { return NodeName; }
+private:
+	FDetailNodeList ChildNodes;
+	IDetailsViewPrivate& DetailsView;
+	TWeakPtr<IDetailRootObjectCustomization> RootObjectCustomization;
+	const TWeakObjectPtr<UObject> RootObject;
+	FName NodeName;
+	bool bShouldBeVisible;
+};
diff --git a/Engine/Source/Editor/PropertyEditor/Private/IDetailTreeNode.h b/Engine/Source/Editor/PropertyEditor/Private/IDetailTreeNode.h
--- a/Engine/Source/Editor/PropertyEditor/Private/IDetailTreeNode.h
+++ b/Engine/Source/Editor/PropertyEditor/Private/IDetailTreeNode.h
@@ -2,17 +2,14 @@
 
 #pragma once
 
-namespace ENodeVisibility
+enum class ENodeVisibility : uint8
 {
-	enum Type
-	{
-		// Hidden but can be visible if parent is visible due to filtering
-		HiddenDueToFiltering,
-		// Never visible no matter what
-		ForcedHidden,
-		// Always visible
-		Visible,
-	};
+	// Hidden but can be visible if parent is visible due to filtering
+	HiddenDueToFiltering,
+	// Never visible no matter what
+	ForcedHidden,
+	// Always visible
+	Visible,
 };
 
 /**
@@ -60,7 +57,7 @@ public:
 	/**
 	 * @return the visibility of this node in the tree
 	 */
-	virtual ENodeVisibility::Type GetVisibility() const = 0;
+	virtual ENodeVisibility GetVisibility() const = 0;
 
 	/**
 	 * Called each frame if the node requests that it should be ticked                                                              
diff --git a/Engine/Source/Editor/PropertyEditor/Private/IDetailsViewPrivate.h b/Engine/Source/Editor/PropertyEditor/Private/IDetailsViewPrivate.h
--- a/Engine/Source/Editor/PropertyEditor/Private/IDetailsViewPrivate.h
+++ b/Engine/Source/Editor/PropertyEditor/Private/IDetailsViewPrivate.h
@@ -14,12 +14,12 @@ class IDetailsViewPrivate : public IDetailsView
 {
 public:
 	/**
-	* Sets the expansion state for a node and optionally all of its children
-	*
-	* @param InTreeNode		The node to change expansion state on
-	* @param bIsItemExpanded	The new expansion state
-	* @param bRecursive		Whether or not to apply the expansion change to any children
-	*/
+	 * Sets the expansion state for a node and optionally all of its children
+	 *
+	 * @param InTreeNode		The node to change expansion state on
+	 * @param bIsItemExpanded	The new expansion state
+	 * @param bRecursive		Whether or not to apply the expansion change to any children
+	 */
 	virtual void SetNodeExpansionState(TSharedRef<IDetailTreeNode> InTreeNode, bool bIsItemExpanded, bool bRecursive) = 0;
 
 	/**
@@ -39,49 +39,37 @@ public:
 	virtual FNotifyHook* GetNotifyHook() const = 0;
 
 	/**
-	* Returns the property utilities for this view
-	*/
+	 * Returns the property utilities for this view
+	 */
 	virtual TSharedPtr<IPropertyUtilities> GetPropertyUtilities() = 0;
 
-
-	/** Gets the base class being viewed */
-	virtual const UClass* GetBaseClass() const = 0;
-	virtual UClass* GetBaseClass() = 0;
-
 	/**
-	* @return True if a category is hidden by any of the uobject classes currently in view by this details panel
-	*/
-	virtual bool IsCategoryHiddenByClass(FName CategoryName) const = 0;
+	 * @return True if a category is hidden by any of the uobject classes currently in view by this details panel
+	 */
+	virtual bool IsCategoryHiddenByClass(const TSharedPtr<FComplexPropertyNode>& InRootNode, FName CategoryName) const = 0;
 
 	/** Causes the details view to be refreshed (new widgets generated) with the current set of objects */
 	virtual void ForceRefresh() = 0;
 
 	/** Causes the details view to move the scroll offset (by item)
-	* @param DeltaOffset	We add this value to the current scroll offset if the result is in the scrolling range
-	*/
+	 * @param DeltaOffset	We add this value to the current scroll offset if the result is in the scrolling range
+	 */
 	virtual void MoveScrollOffset(int32 DeltaOffset) = 0;
 
-	/* Return true if the category exist in the detail view
-	* @param CategoryName			The name of the category we need informations
-	* @param SimplePropertiesNum		The Number of basic properties contain in the category
-	* @param AdvancePropertiesNum	The Number of advance properties contain in the category
-	*/
-	virtual bool GetCategoryInfo(FName CategoryName, int32 &SimplePropertiesNum, int32 &AdvancePropertiesNum) { SimplePropertiesNum = 0; AdvancePropertiesNum = 0; return false; }
-
 	/**
-	* Saves the expansion state of a tree node
-	*
-	* @param NodePath	The path to the detail node to save
-	* @param bIsExpanded	true if the node is expanded, false otherwise
-	*/
+	 * Saves the expansion state of a tree node
+	 *
+	 * @param NodePath	The path to the detail node to save
+	 * @param bIsExpanded	true if the node is expanded, false otherwise
+	 */
 	virtual void SaveCustomExpansionState(const FString& NodePath, bool bIsExpanded) = 0;
 
 	/**
-	* Gets the saved expansion state of a tree node in this category
-	*
-	* @param NodePath	The path to the detail node to get
-	* @return true if the node should be expanded, false otherwise
-	*/
+	 * Gets the saved expansion state of a tree node in this category
+	 *
+	 * @param NodePath	The path to the detail node to get
+	 * @return true if the node should be expanded, false otherwise
+	 */
 	virtual bool GetCustomSavedExpansionState(const FString& NodePath) const = 0;
 
 	/**
@@ -133,4 +121,10 @@ public:
 
 	/** If a customization standalone widget is used, the value should be update only once, when its window is closed */
 	virtual bool DontUpdateValueWhileEditing() const = 0;
+
+	/** @return Whether or not the details panel was created with multiple unrelated objects visible at once in the details panel */
+	virtual bool ContainsMultipleTopLevelObjects() const = 0;
+
+	/** @return the customization instance that defines how the display for a root object looks */
+	virtual TSharedPtr<class IDetailRootObjectCustomization> GetRootObjectCustomization() const = 0;
 };
diff --git a/Engine/Source/Editor/PropertyEditor/Private/ObjectPropertyNode.cpp b/Engine/Source/Editor/PropertyEditor/Private/ObjectPropertyNode.cpp
--- a/Engine/Source/Editor/PropertyEditor/Private/ObjectPropertyNode.cpp
+++ b/Engine/Source/Editor/PropertyEditor/Private/ObjectPropertyNode.cpp
@@ -50,7 +50,7 @@ const UPackage* FObjectPropertyNode::GetUPackage(int32 InIndex) const
 // Adds a new object to the list.
 void FObjectPropertyNode::AddObject( UObject* InObject )
 {
-	Objects.AddUnique( InObject );
+	Objects.Add( InObject );
 }
 
 // Removes an object from the list.
@@ -115,7 +115,7 @@ void FObjectPropertyNode::Finalize()
 
 bool FObjectPropertyNode::GetReadAddressUncached(FPropertyNode& InNode,
 											   bool InRequiresSingleSelection,
-											   FReadAddressListData& OutAddresses,
+											   FReadAddressListData* OutAddresses,
 											   bool bComparePropertyContents,
 											   bool bObjectForceCompare,
 											   bool bArrayPropertiesCanDifferInSize) const
@@ -238,13 +238,16 @@ bool FObjectPropertyNode::GetReadAddressUncached(FPropertyNode& InNode,
 		}
 	}
 
-	// Write addresses to the output.
-	for ( int32 ObjIndex = 0 ; ObjIndex < GetNumObjects(); ++ObjIndex )
+	if(OutAddresses != nullptr)
 	{
-		const UObject* TempObject = GetUObject(ObjIndex);
-		if( TempObject )
+		// Write addresses to the output.
+		for(int32 ObjIndex = 0; ObjIndex < GetNumObjects(); ++ObjIndex)
 		{
-			OutAddresses.Add( TempObject, InNode.GetValueBaseAddress( (uint8*)(TempObject) ) );
+			const UObject* TempObject = GetUObject(ObjIndex);
+			if(TempObject)
+			{
+				OutAddresses->Add(TempObject, InNode.GetValueBaseAddress((uint8*)(TempObject)));
+			}
 		}
 	}
 
diff --git a/Engine/Source/Editor/PropertyEditor/Private/ObjectPropertyNode.h b/Engine/Source/Editor/PropertyEditor/Private/ObjectPropertyNode.h
--- a/Engine/Source/Editor/PropertyEditor/Private/ObjectPropertyNode.h
+++ b/Engine/Source/Editor/PropertyEditor/Private/ObjectPropertyNode.h
@@ -20,7 +20,7 @@ public:
 	/** FPropertyNode Interface */
 	virtual FObjectPropertyNode* AsObjectNode() override { return this;}
 	virtual const FObjectPropertyNode* AsObjectNode() const override { return this; }
-	virtual bool GetReadAddressUncached(FPropertyNode& InNode, bool InRequiresSingleSelection, FReadAddressListData& OutAddresses, bool bComparePropertyContents = true, bool bObjectForceCompare = false, bool bArrayPropertiesCanDifferInSize = false) const override;
+	virtual bool GetReadAddressUncached(FPropertyNode& InNode, bool InRequiresSingleSelection, FReadAddressListData* OutAddresses, bool bComparePropertyContents = true, bool bObjectForceCompare = false, bool bArrayPropertiesCanDifferInSize = false) const override;
 	virtual bool GetReadAddressUncached(FPropertyNode& InNode, FReadAddressListData& OutAddresses) const override;
 
 	virtual uint8* GetValueBaseAddress( uint8* Base ) override;
@@ -109,6 +109,7 @@ public:
 	 */
 	const TSet<FName>& GetHiddenCategories() const { return HiddenCategories; }
 
+	bool IsRootNode() const { return ParentNode == nullptr; }
 protected:
 	/** FPropertyNode interface */
 	virtual void InitBeforeNodeFlags() override;
diff --git a/Engine/Source/Editor/PropertyEditor/Private/Presentation/PropertyTable/PropertyTableColumn.cpp b/Engine/Source/Editor/PropertyEditor/Private/Presentation/PropertyTable/PropertyTableColumn.cpp
--- a/Engine/Source/Editor/PropertyEditor/Private/Presentation/PropertyTable/PropertyTableColumn.cpp
+++ b/Engine/Source/Editor/PropertyEditor/Private/Presentation/PropertyTable/PropertyTableColumn.cpp
@@ -683,13 +683,13 @@ void FPropertyTableColumn::Tick()
 		else
 		{
 			const TSharedRef< FObjectPropertyNode > Node = TableRef->GetObjectPropertyNode( Object );
-			FPropertyNode::DataValidationResult Result = Node->EnsureDataIsValid();
+			EPropertyDataValidationResult Result = Node->EnsureDataIsValid();
 
-			if ( Result == FPropertyNode::ObjectInvalid )
+			if ( Result == EPropertyDataValidationResult::ObjectInvalid )
 			{
 				TableRef->RemoveColumn( SharedThis( this ) );
 			}
-			else if ( Result == FPropertyNode::ArraySizeChanged )
+			else if ( Result == EPropertyDataValidationResult::ArraySizeChanged )
 			{
 				TableRef->RequestRefresh();
 			}
diff --git a/Engine/Source/Editor/PropertyEditor/Private/Presentation/PropertyTable/PropertyTableRow.cpp b/Engine/Source/Editor/PropertyEditor/Private/Presentation/PropertyTable/PropertyTableRow.cpp
--- a/Engine/Source/Editor/PropertyEditor/Private/Presentation/PropertyTable/PropertyTableRow.cpp
+++ b/Engine/Source/Editor/PropertyEditor/Private/Presentation/PropertyTable/PropertyTableRow.cpp
@@ -139,17 +139,17 @@ void FPropertyTableRow::Tick()
 				return;
 			}
 
-			FPropertyNode::DataValidationResult Result = Node->EnsureDataIsValid();
+			EPropertyDataValidationResult Result = Node->EnsureDataIsValid();
 
-			if ( Result == FPropertyNode::ObjectInvalid )
+			if ( Result == EPropertyDataValidationResult::ObjectInvalid )
 			{
 				TableRef->RemoveRow( SharedThis( this ) );
 			}
-			else if ( Result == FPropertyNode::ArraySizeChanged )
+			else if ( Result == EPropertyDataValidationResult::ArraySizeChanged )
 			{
 				TableRef->RequestRefresh();
 			}
-			else if ( Result != FPropertyNode::DataValid )
+			else if ( Result != EPropertyDataValidationResult::DataValid )
 			{
 				Refresh();
 			}
diff --git a/Engine/Source/Editor/PropertyEditor/Private/PropertyHandleImpl.cpp b/Engine/Source/Editor/PropertyEditor/Private/PropertyHandleImpl.cpp
--- a/Engine/Source/Editor/PropertyEditor/Private/PropertyHandleImpl.cpp
+++ b/Engine/Source/Editor/PropertyEditor/Private/PropertyHandleImpl.cpp
@@ -369,6 +369,10 @@ FPropertyAccess::Result FPropertyValueImpl::ImportText( const TArray<FObjectBase
 		const bool bTransactable = (Flags & EPropertyValueSetFlags::NotTransactable) == 0;
 		const bool bFinished = ( Flags & EPropertyValueSetFlags::InteractiveChange) == 0;
 		
+		// List of top level objects sent to the PropertyChangedEvent
+		TArray<const UObject*> TopLevelObjects;
+		TopLevelObjects.Reserve(InObjects.Num());
+
 		for ( int32 ObjectIndex = 0 ; ObjectIndex < InObjects.Num() ; ++ObjectIndex )
 		{	
 			const FObjectBaseAddress& Cur = InObjects[ ObjectIndex ];
@@ -445,6 +449,8 @@ FPropertyAccess::Result FPropertyValueImpl::ImportText( const TArray<FObjectBase
 				{
 					Cur.Object->MarkPackageDirty();
 				}
+
+				TopLevelObjects.Add(Cur.Object);
 			}
 
 			//add on array index so we can tell which entry just changed
@@ -452,13 +458,13 @@ FPropertyAccess::Result FPropertyValueImpl::ImportText( const TArray<FObjectBase
 			FPropertyValueImpl::GenerateArrayIndexMapToObjectNode( ArrayIndicesPerObject[ObjectIndex], InPropertyNode );
 		}
 
+		FPropertyChangedEvent ChangeEvent(NodeProperty, bFinished ? EPropertyChangeType::ValueSet : EPropertyChangeType::Interactive, &TopLevelObjects);
+		ChangeEvent.SetArrayIndexPerObject(ArrayIndicesPerObject);
+
 		// If PreEditChange was called, so should PostEditChange.
-		if ( bNotifiedPreChange )
+		if (bNotifiedPreChange)
 		{
 			// Call PostEditChange on all objects.
-			FPropertyChangedEvent ChangeEvent(NodeProperty, bFinished ? EPropertyChangeType::ValueSet : EPropertyChangeType::Interactive );
-			ChangeEvent.SetArrayIndexPerObject(ArrayIndicesPerObject);
-
 			InPropertyNode->NotifyPostChange( ChangeEvent, NotifyHook );
 
 			if (bFinished)
@@ -481,7 +487,6 @@ FPropertyAccess::Result FPropertyValueImpl::ImportText( const TArray<FObjectBase
 
 		if (PropertyUtilities.IsValid() && !bInteractiveChangeInProgress)
 		{
-			FPropertyChangedEvent ChangeEvent(NodeProperty, bFinished ? EPropertyChangeType::ValueSet : EPropertyChangeType::Interactive);
 			InPropertyNode->FixPropertiesInEvent(ChangeEvent);
 			PropertyUtilities.Pin()->NotifyFinishedChangingProperties(ChangeEvent);
 		}
@@ -1082,6 +1087,10 @@ void FPropertyValueImpl::AddChild()
 
 			TArray< TMap<FString,int32> > ArrayIndicesPerObject;
 
+			// List of top level objects sent to the PropertyChangedEvent
+			TArray<const UObject*> TopLevelObjects;
+			TopLevelObjects.Reserve(ReadAddresses.Num());
+
 			// Begin a property edit transaction.
 			FScopedTransaction Transaction( NSLOCTEXT("UnrealEd", "AddChild", "Add Child") );
 			FObjectPropertyNode* ObjectNode = PropertyNodePin->FindObjectItemParent();
@@ -1115,6 +1124,8 @@ void FPropertyValueImpl::AddChild()
 
 							PropertyNodePin->PropagateArrayPropertyChange(Obj, OrgArrayContent, EPropertyArrayChangeType::Add, -1);
 						}
+
+						TopLevelObjects.Add(Obj);
 					}
 
 					FScriptArrayHelper	ArrayHelper(Array,Addr);
@@ -1125,17 +1136,17 @@ void FPropertyValueImpl::AddChild()
 				}
 			}
 
+			FPropertyChangedEvent ChangeEvent(NodeProperty, EPropertyChangeType::ArrayAdd, &TopLevelObjects);
+			ChangeEvent.SetArrayIndexPerObject(ArrayIndicesPerObject);
+
 			if ( bNotifiedPreChange )
 			{
 				// send the PostEditChange notification; it will be propagated to all selected objects
-				FPropertyChangedEvent ChangeEvent(NodeProperty, EPropertyChangeType::ArrayAdd);
-				ChangeEvent.SetArrayIndexPerObject(ArrayIndicesPerObject);
-				PropertyNodePin->NotifyPostChange( ChangeEvent, NotifyHook );
+				PropertyNodePin->NotifyPostChange( ChangeEvent, NotifyHook);
 			}
 
 			if (PropertyUtilities.IsValid())
 			{
-				FPropertyChangedEvent ChangeEvent(NodeProperty, EPropertyChangeType::ArrayAdd);
 				PropertyNodePin->FixPropertiesInEvent(ChangeEvent);
 				PropertyUtilities.Pin()->NotifyFinishedChangingProperties(ChangeEvent);
 			}
@@ -1162,6 +1173,10 @@ void FPropertyValueImpl::ClearChildren()
 			// we don't want the objects to be marked dirty)
 			bool bNotifiedPreChange = false;
 
+			// List of top level objects sent to the PropertyChangedEvent
+			TArray<const UObject*> TopLevelObjects;
+			TopLevelObjects.Reserve(ReadAddresses.Num());
+
 			// Begin a property edit transaction.
 			FScopedTransaction Transaction( NSLOCTEXT("UnrealEd", "ClearChildren", "Clear Children") );
 			FObjectPropertyNode* ObjectNode = PropertyNodePin->FindObjectItemParent();
@@ -1191,6 +1206,8 @@ void FPropertyValueImpl::ClearChildren()
 
 							PropertyNodePin->PropagateArrayPropertyChange(Obj, OrgArrayContent, EPropertyArrayChangeType::Clear, -1);
 						}
+
+						TopLevelObjects.Add(Obj);
 					}
 
 					FScriptArrayHelper	ArrayHelper(Array,Addr);
@@ -1198,19 +1215,19 @@ void FPropertyValueImpl::ClearChildren()
 				}
 			}
 
+			FPropertyChangedEvent ChangeEvent(NodeProperty, EPropertyChangeType::ValueSet, &TopLevelObjects);
+
 			if ( bNotifiedPreChange )
 			{
 				// send the PostEditChange notification; it will be propagated to all selected objects
-				FPropertyChangedEvent ChangeEvent(NodeProperty, EPropertyChangeType::ValueSet);
-				PropertyNodePin->NotifyPostChange( ChangeEvent, NotifyHook );
+				PropertyNodePin->NotifyPostChange( ChangeEvent, NotifyHook);
 			}
-		}
 
-		if (PropertyUtilities.IsValid())
-		{
-			FPropertyChangedEvent ChangeEvent(NodeProperty, EPropertyChangeType::ValueSet);
-			PropertyNodePin->FixPropertiesInEvent(ChangeEvent);
-			PropertyUtilities.Pin()->NotifyFinishedChangingProperties(ChangeEvent);
+			if(PropertyUtilities.IsValid())
+			{
+				PropertyNodePin->FixPropertiesInEvent(ChangeEvent);
+				PropertyUtilities.Pin()->NotifyFinishedChangingProperties(ChangeEvent);
+			}
 		}
 	}
 }
@@ -1251,6 +1268,9 @@ void FPropertyValueImpl::InsertChild( TSharedPtr<FPropertyNode> ChildNodeToInser
 		FScriptArrayHelper	ArrayHelper(ArrayProperty,Addr);
 		int32 Index = ChildNodePtr->GetArrayIndex();
 
+		// List of top level objects sent to the PropertyChangedEvent
+		TArray<const UObject*> TopLevelObjects;
+
 		UObject* Obj = ObjectNode ? ObjectNode->GetUObject(0) : nullptr;
 		if (Obj)
 		{
@@ -1263,6 +1283,8 @@ void FPropertyValueImpl::InsertChild( TSharedPtr<FPropertyNode> ChildNodeToInser
 
 				ChildNodePtr->PropagateArrayPropertyChange(Obj, OrgArrayContent, EPropertyArrayChangeType::Insert, Index);
 			}
+
+			TopLevelObjects.Add(Obj);
 		}
 
 		ArrayHelper.InsertValues(Index, 1 );
@@ -1278,15 +1300,12 @@ void FPropertyValueImpl::InsertChild( TSharedPtr<FPropertyNode> ChildNodeToInser
 			FPropertyValueImpl::GenerateArrayIndexMapToObjectNode(ArrayIndicesPerObject[ObjectIndex], ChildNodePtr );
 		}
 
-		{
-			FPropertyChangedEvent ChangeEvent(ParentNode->GetProperty(), EPropertyChangeType::ArrayAdd);
-			ChangeEvent.SetArrayIndexPerObject(ArrayIndicesPerObject);
-			ChildNodePtr->NotifyPostChange(ChangeEvent, NotifyHook);
-		}
+		FPropertyChangedEvent ChangeEvent(ParentNode->GetProperty(), EPropertyChangeType::ArrayAdd, &TopLevelObjects);
+		ChangeEvent.SetArrayIndexPerObject(ArrayIndicesPerObject);
+		ChildNodePtr->NotifyPostChange(ChangeEvent, NotifyHook);
 
 		if (PropertyUtilities.IsValid())
 		{
-			FPropertyChangedEvent ChangeEvent(ParentNode->GetProperty(), EPropertyChangeType::ArrayAdd);
 			ChildNodePtr->FixPropertiesInEvent(ChangeEvent);
 			PropertyUtilities.Pin()->NotifyFinishedChangingProperties(ChangeEvent);
 		}
@@ -1321,6 +1340,10 @@ void FPropertyValueImpl::DeleteChild( TSharedPtr<FPropertyNode> ChildNodeToDelet
 
 		ChildNodePtr->NotifyPreChange( NodeProperty, NotifyHook );
 
+		// List of top level objects sent to the PropertyChangedEvent
+		TArray<const UObject*> TopLevelObjects;
+		TopLevelObjects.Reserve(ReadAddresses.Num());
+
 		// perform the operation on the array for all selected objects
 		for ( int32 i = 0 ; i < ReadAddresses.Num() ; ++i )
 		{
@@ -1343,20 +1366,19 @@ void FPropertyValueImpl::DeleteChild( TSharedPtr<FPropertyNode> ChildNodeToDelet
 
 						ChildNodePtr->PropagateArrayPropertyChange(Obj, OrgArrayContent, EPropertyArrayChangeType::Delete, Index);
 					}
+
+					TopLevelObjects.Add(Obj);
 				}
 
 				ArrayHelper.RemoveValues(ChildNodePtr->GetArrayIndex());
 			}
 		}
 
-		{
-			FPropertyChangedEvent ChangeEvent(ParentNode->GetProperty());
-			ChildNodePtr->NotifyPostChange(ChangeEvent, NotifyHook);
-		}
+		FPropertyChangedEvent ChangeEvent(ParentNode->GetProperty(), EPropertyChangeType::Unspecified, &TopLevelObjects);
+		ChildNodePtr->NotifyPostChange(ChangeEvent, NotifyHook);
 
 		if (PropertyUtilities.IsValid())
 		{
-			FPropertyChangedEvent ChangeEvent(ParentNode->GetProperty());
 			ChildNodePtr->FixPropertiesInEvent(ChangeEvent);
 			PropertyUtilities.Pin()->NotifyFinishedChangingProperties(ChangeEvent);
 		}
@@ -1392,6 +1414,9 @@ void FPropertyValueImpl::DuplicateChild( TSharedPtr<FPropertyNode> ChildNodeToDu
 
 	if( Addr )
 	{
+		// List of top level objects sent to the PropertyChangedEvent
+		TArray<const UObject*> TopLevelObjects;
+
 		FScopedTransaction Transaction( NSLOCTEXT("UnrealEd", "DuplicateChild", "Duplicate Child") );
 
 		ChildNodePtr->NotifyPreChange( ParentNode->GetProperty(), NotifyHook );
@@ -1413,6 +1438,8 @@ void FPropertyValueImpl::DuplicateChild( TSharedPtr<FPropertyNode> ChildNodeToDu
 
 				ChildNodePtr->PropagateArrayPropertyChange(Obj, OrgArrayContent, EPropertyArrayChangeType::Duplicate, Index);
 			}
+
+			TopLevelObjects.Add(Obj);
 		}
 
 		ArrayHelper.InsertValues(Index);
@@ -1454,19 +1481,11 @@ void FPropertyValueImpl::DuplicateChild( TSharedPtr<FPropertyNode> ChildNodeToDu
 			FPropertyValueImpl::GenerateArrayIndexMapToObjectNode(ArrayIndicesPerObject[0], ChildNodePtr);
 		}
 
+		FPropertyChangedEvent ChangeEvent(ParentNode->GetProperty(), EPropertyChangeType::ValueSet, &TopLevelObjects);
+		ChangeEvent.SetArrayIndexPerObject(ArrayIndicesPerObject);
 
+		ChildNodePtr->NotifyPostChange(ChangeEvent, NotifyHook);
 
-		//@todo Slate Property Window
-		// 		const bool bExpand = true;
-		// 		const bool bRecurse = false;
-		// 		ParentNode->SetExpanded(bExpand, bRecurse);
-
-		FPropertyChangedEvent ChangeEvent(ParentNode->GetProperty(), EPropertyChangeType::ValueSet);
-		ChangeEvent.SetArrayIndexPerObject(ArrayIndicesPerObject);
-		{
-			
-			ChildNodePtr->NotifyPostChange(ChangeEvent, NotifyHook);
-		}
 		if (PropertyUtilities.IsValid())
 		{
 			ChildNodePtr->FixPropertiesInEvent(ChangeEvent);
@@ -2163,8 +2182,19 @@ void FPropertyHandleBase::NotifyPostChange()
 	TSharedPtr<FPropertyNode> PropertyNode = Implementation->GetPropertyNode();
 	if( PropertyNode.IsValid() )
 	{
-		FPropertyChangedEvent PropertyChangedEvent( PropertyNode->GetProperty() );
-		PropertyNode->NotifyPostChange( PropertyChangedEvent, Implementation->GetNotifyHook() );
+		TArray<const UObject*> ObjectsBeingChanged;
+		FObjectPropertyNode* ObjectNode = PropertyNode->FindObjectItemParent();
+		if(ObjectNode)
+		{
+			ObjectsBeingChanged.Reserve(ObjectNode->GetNumObjects());
+			for(int32 ObjectIndex = 0; ObjectIndex < ObjectNode->GetNumObjects(); ++ObjectIndex)
+			{
+				ObjectsBeingChanged.Add(ObjectNode->GetUObject(ObjectIndex));
+			}
+		}
+
+		FPropertyChangedEvent PropertyChangedEvent( PropertyNode->GetProperty(), EPropertyChangeType::Unspecified, &ObjectsBeingChanged );
+		PropertyNode->NotifyPostChange( PropertyChangedEvent, Implementation->GetNotifyHook());
 	}
 }
 
diff --git a/Engine/Source/Editor/PropertyEditor/Private/PropertyNode.cpp b/Engine/Source/Editor/PropertyEditor/Private/PropertyNode.cpp
--- a/Engine/Source/Editor/PropertyEditor/Private/PropertyNode.cpp
+++ b/Engine/Source/Editor/PropertyEditor/Private/PropertyNode.cpp
@@ -125,9 +125,8 @@ void FPropertyNode::InitNode( const FPropertyNodeInitParams& InitParams )
 	}
 	else
 	{
-		FReadAddressListData ReadAddresses;
-		const bool GotReadAddresses = GetReadAddressUncached( *this, false, ReadAddresses, false );
-		const bool bSingleSelectOnly = GetReadAddressUncached( *this, true, ReadAddresses );
+		const bool GotReadAddresses = GetReadAddressUncached( *this, false, nullptr, false );
+		const bool bSingleSelectOnly = GetReadAddressUncached( *this, true, nullptr);
 		SetNodeFlags(EPropertyNodeFlags::SingleSelectOnly, bSingleSelectOnly);
 
 		UProperty* MyProperty = Property.Get();
@@ -329,7 +328,7 @@ FObjectPropertyNode* FPropertyNode::FindRootObjectItemParent()
 /** 
  * Used to see if any data has been destroyed from under the property tree.  Should only be called by PropertyWindow::OnIdle
  */
-FPropertyNode::DataValidationResult FPropertyNode::EnsureDataIsValid()
+EPropertyDataValidationResult FPropertyNode::EnsureDataIsValid()
 {
 	bool bValidateChildren = !HasNodeFlags(EPropertyNodeFlags::SkipChildValidation);
 
@@ -391,7 +390,7 @@ FPropertyNode::DataValidationResult FPropertyNode::EnsureDataIsValid()
 			if (!bSuccess)
 			{
 				UE_LOG( LogPropertyNode, Verbose, TEXT("Object is invalid %s"), *Property->GetName() );
-				return ObjectInvalid;
+				return EPropertyDataValidationResult::ObjectInvalid;
 			}
 
 
@@ -403,7 +402,7 @@ FPropertyNode::DataValidationResult FPropertyNode::EnsureDataIsValid()
 				if (Addr==NULL)
 				{
 					UE_LOG( LogPropertyNode, Verbose, TEXT("Object is invalid %s"), *Property->GetName() );
-					return ObjectInvalid;
+					return EPropertyDataValidationResult::ObjectInvalid;
 				}
 
 				if( ArrayProperty && !bIgnoreAllMismatch)
@@ -449,7 +448,7 @@ FPropertyNode::DataValidationResult FPropertyNode::EnsureDataIsValid()
 					}
 				}
 
-				return ArraySizeChanged;
+				return EPropertyDataValidationResult::ArraySizeChanged;
 			}
 
 			const bool bHasChildren = (GetNumChildNodes() != 0);
@@ -458,7 +457,7 @@ FPropertyNode::DataValidationResult FPropertyNode::EnsureDataIsValid()
 			if (ObjectProperty && ((!bObjectPropertyNull && !bHasChildren) || (bObjectPropertyNull && bHasChildren)))
 			{
 				RebuildChildren();
-				return PropertiesChanged;
+				return EPropertyDataValidationResult::PropertiesChanged;
 			}
 		}
 	}
@@ -467,10 +466,10 @@ FPropertyNode::DataValidationResult FPropertyNode::EnsureDataIsValid()
 	{
 		RebuildChildren();
 		// If this property is editinline and not edit const then its editinline new and we can optimize some of the refreshing in some cases.  Otherwise we need to refresh all properties in the view
-		return HasNodeFlags(EPropertyNodeFlags::EditInline) && !IsEditConst() ? EditInlineNewValueChanged : PropertiesChanged;
+		return HasNodeFlags(EPropertyNodeFlags::EditInline) && !IsEditConst() ? EPropertyDataValidationResult::EditInlineNewValueChanged : EPropertyDataValidationResult::PropertiesChanged;
 	}
 	
-	FPropertyNode::DataValidationResult FinalResult = DataValid;
+	EPropertyDataValidationResult FinalResult = EPropertyDataValidationResult::DataValid;
 
 	//go through my children
 	if (bValidateChildren)
@@ -480,8 +479,8 @@ FPropertyNode::DataValidationResult FPropertyNode::EnsureDataIsValid()
 			TSharedPtr<FPropertyNode>& ChildNode = ChildNodes[Scan];
 			check(ChildNode.IsValid());
 
-			FPropertyNode::DataValidationResult ChildDataResult = ChildNode->EnsureDataIsValid();
-			if (FinalResult == DataValid && ChildDataResult != DataValid)
+			EPropertyDataValidationResult ChildDataResult = ChildNode->EnsureDataIsValid();
+			if (FinalResult == EPropertyDataValidationResult::DataValid && ChildDataResult != EPropertyDataValidationResult::DataValid)
 			{
 				FinalResult = ChildDataResult;
 			}
@@ -625,7 +624,7 @@ bool FPropertyNode::GetQualifiedName( FString& PathPlusIndex, const bool bWithAr
 
 bool FPropertyNode::GetReadAddressUncached( FPropertyNode& InPropertyNode,
 									bool InRequiresSingleSelection,
-									FReadAddressListData& OutAddresses,
+									FReadAddressListData* OutAddresses,
 									bool bComparePropertyContents,
 									bool bObjectForceCompare,
 									bool bArrayPropertiesCanDifferInSize ) const
@@ -667,7 +666,7 @@ bool FPropertyNode::GetReadAddress(bool InRequiresSingleSelection,
 	bool bAllValuesTheSame = false;
 	if (ParentNodeWeakPtr.IsValid())
 	{
-		bAllValuesTheSame = GetReadAddressUncached( *this, InRequiresSingleSelection, CachedReadAddresses, bComparePropertyContents, bObjectForceCompare, bArrayPropertiesCanDifferInSize );
+		bAllValuesTheSame = GetReadAddressUncached( *this, InRequiresSingleSelection, &CachedReadAddresses, bComparePropertyContents, bObjectForceCompare, bArrayPropertiesCanDifferInSize );
 		OutAddresses.ReadAddressListData = &CachedReadAddresses;
 		CachedReadAddresses.bAllValuesTheSame = bAllValuesTheSame;
 		CachedReadAddresses.bRequiresCache = false;
diff --git a/Engine/Source/Editor/PropertyEditor/Private/PropertyNode.h b/Engine/Source/Editor/PropertyEditor/Private/PropertyNode.h
--- a/Engine/Source/Editor/PropertyEditor/Private/PropertyNode.h
+++ b/Engine/Source/Editor/PropertyEditor/Private/PropertyNode.h
@@ -110,7 +110,7 @@ public:
 	uint8* GetAddress( int32 Index )
 	{
 		const FAddressPair& Pair = ReadAddresses[Index];
-		return (Pair.Object.IsValid() || Pair.bIsStruct) ? Pair.ReadAddress : NULL;
+		return (Pair.Object.IsValid() || Pair.bIsStruct) ? Pair.ReadAddress : 0;
 	}
 
 	bool IsValidIndex( int32 Index ) const
@@ -139,12 +139,12 @@ class FReadAddressList
 	friend class FPropertyNode;
 public:
 	FReadAddressList()
-		: ReadAddressListData( NULL )
+		: ReadAddressListData(nullptr)
 	{}
 
 	int32 Num() const
 	{
-		return (ReadAddressListData != NULL) ? ReadAddressListData->Num() : 0;
+		return (ReadAddressListData != nullptr) ? ReadAddressListData->Num() : 0;
 	}
 
 	uint8* GetAddress( int32 Index )
@@ -159,7 +159,7 @@ public:
 
 	void Reset()
 	{
-		if (ReadAddressListData != NULL)
+		if (ReadAddressListData != nullptr)
 		{
 			ReadAddressListData->Reset();
 		}
@@ -194,8 +194,8 @@ struct FPropertyNodeInitParams
 	bool bCreateDisableEditOnInstanceNodes;
 
 	FPropertyNodeInitParams()
-		: ParentNode( NULL )
-		, Property( NULL )
+		: ParentNode(nullptr)
+		, Property(nullptr)
 		, ArrayOffset(0)
 		, ArrayIndex( INDEX_NONE )
 		, bAllowChildren( true )
@@ -226,6 +226,20 @@ struct EPropertyArrayChangeType
 
 class FComplexPropertyNode;
 
+enum EPropertyDataValidationResult : uint8
+{
+	/** The object(s) being viewed are now invalid */
+	ObjectInvalid,
+	/** Non dynamic array property nodes were added or removed that would require a refresh */
+	PropertiesChanged,
+	/** An edit inline new value changed,  In the tree this rebuilds the nodes, in the details view we don't need to do this */
+	EditInlineNewValueChanged,
+	/** The size of an array changed (delete,insert,add) */
+	ArraySizeChanged,
+	/** All data is valid */
+	DataValid,
+};
+
 /**
  * The base class for all property nodes                                                              
  */
@@ -233,20 +247,6 @@ class FPropertyNode : public TSharedFromThis<FPropertyNode>
 {
 public:
 
-	enum DataValidationResult
-	{
-		/** The object(s) being viewed are now invalid */
-		ObjectInvalid,
-		/** Non dynamic array property nodes were added or removed that would require a refresh */
-		PropertiesChanged,
-		/** An edit inline new value changed,  In the tree this rebuilds the nodes, in the details view we don't need to do this */
-		EditInlineNewValueChanged,
-		/** The size of an array changed (delete,insert,add) */
-		ArraySizeChanged,
-		/** All data is valid */
-		DataValid,
-	};
-
 	FPropertyNode(void);
 	virtual ~FPropertyNode(void);
 
@@ -279,21 +279,21 @@ public:
 	/**
 	 * Interface function to get at the dervied FObjectPropertyNodeWx class
 	 */
-	virtual class FObjectPropertyNode* AsObjectNode() { return NULL; }
-	virtual const FObjectPropertyNode* AsObjectNode() const { return NULL; }
+	virtual class FObjectPropertyNode* AsObjectNode() { return nullptr; }
+	virtual const FObjectPropertyNode* AsObjectNode() const { return nullptr; }
 
-	virtual FComplexPropertyNode* AsComplexNode() { return NULL; }
-	virtual const FComplexPropertyNode* AsComplexNode() const { return NULL; }
+	virtual FComplexPropertyNode* AsComplexNode() { return nullptr; }
+	virtual const FComplexPropertyNode* AsComplexNode() const { return nullptr; }
 
 	/**
 	 * Interface function to get at the dervied FCategoryPropertyNodeWx class
 	 */
-	virtual class FCategoryPropertyNode* AsCategoryNode() { return NULL; }
+	virtual class FCategoryPropertyNode* AsCategoryNode() { return nullptr; }
 
 	/**
 	 * Interface function to get at the dervied FItemPropertyNodeWx class
 	 */
-	virtual class FItemPropertyNode* AsItemPropertyNode() { return NULL; }
+	virtual class FItemPropertyNode* AsItemPropertyNode() { return nullptr; }
 
 	/**
 	 * Follows the chain of items upwards until it finds the object window that houses this item.
@@ -312,7 +312,7 @@ public:
 	/** 
 	 * Used to see if any data has been destroyed from under the property tree.  Should only be called during Tick
 	 */
-	FPropertyNode::DataValidationResult EnsureDataIsValid();
+	EPropertyDataValidationResult EnsureDataIsValid();
 
 	//////////////////////////////////////////////////////////////////////////
 	//Flags
@@ -335,8 +335,8 @@ public:
 	/**
 	 * Returns the parent node in the hierarchy
 	 */
-	FPropertyNode*			GetParentNode()			{ return ParentNodeWeakPtr.IsValid() ? ParentNodeWeakPtr.Pin().Get() : NULL; }
-	const FPropertyNode*	GetParentNode() const	{ return ParentNodeWeakPtr.IsValid() ? ParentNodeWeakPtr.Pin().Get() : NULL; }
+	FPropertyNode*			GetParentNode()			{ return ParentNodeWeakPtr.IsValid() ? ParentNodeWeakPtr.Pin().Get() : nullptr; }
+	const FPropertyNode*	GetParentNode() const	{ return ParentNodeWeakPtr.IsValid() ? ParentNodeWeakPtr.Pin().Get() : nullptr; }
 	TSharedPtr<FPropertyNode> GetParentNodeSharedPtr() { return ParentNodeWeakPtr.Pin(); }
 	/**
 	 * Returns the Property this Node represents
@@ -403,7 +403,7 @@ public:
 	/** 
 	 * Gets read addresses without accessing cached data.  Is less efficient but gets the must up to date data 
 	 */
-	virtual bool GetReadAddressUncached(FPropertyNode& InNode, bool InRequiresSingleSelection, FReadAddressListData& OutAddresses, bool bComparePropertyContents = true, bool bObjectForceCompare = false, bool bArrayPropertiesCanDifferInSize = false) const;
+	virtual bool GetReadAddressUncached(FPropertyNode& InNode, bool InRequiresSingleSelection, FReadAddressListData* OutAddresses, bool bComparePropertyContents = true, bool bObjectForceCompare = false, bool bArrayPropertiesCanDifferInSize = false) const;
 	virtual bool GetReadAddressUncached(FPropertyNode& InNode, FReadAddressListData& OutAddresses) const;
 
 	/**
@@ -546,9 +546,16 @@ public:
 	{
 		TArray< FPropertyInfo > Properties;
 		FPropertyNode* CurrentNode = &PropertyNode.Get();
-		while ( CurrentNode != NULL )
+
+		if(CurrentNode != nullptr && CurrentNode->AsCategoryNode() != nullptr)
 		{
-			if ( CurrentNode->AsItemPropertyNode() != NULL )
+			TSharedRef< FPropertyPath > NewPath = MakeShareable(new FPropertyPath());
+			return NewPath;
+		}
+
+		while (CurrentNode != nullptr)
+		{
+			if (CurrentNode->AsItemPropertyNode() != nullptr)
 			{
 				FPropertyInfo NewPropInfo;
 				NewPropInfo.Property = CurrentNode->GetProperty();
@@ -594,7 +601,7 @@ public:
 				{
 					const TSharedPtr< FPropertyNode > ChildNode = CurrentNode->GetChildNode( ChildIndex );
 
-					if ( ChildNode->AsItemPropertyNode() == NULL )
+					if ( ChildNode->AsItemPropertyNode() == nullptr)
 					{
 						ChildrenStack.Add( ChildNode.ToSharedRef() );
 					}
@@ -612,7 +619,7 @@ public:
 
 		if ( FailedToFindProperty )
 		{
-			PropertyNode = NULL;
+			PropertyNode = nullptr;
 		}
 
 		return PropertyNode;
@@ -633,7 +640,7 @@ public:
 		{
 			TSharedPtr< FPropertyNode > CurrentNode = PropertyNode->GetChildNode( ChildIndex );
 
-			if ( CurrentNode.IsValid() && CurrentNode->AsItemPropertyNode() != NULL )
+			if ( CurrentNode.IsValid() && CurrentNode->AsItemPropertyNode() != nullptr)
 			{
 				FPropertyInfo NewPropInfo;
 				NewPropInfo.Property = CurrentNode->GetProperty();
@@ -743,7 +750,7 @@ protected:
 	void ExpandParent( bool bInRecursive );
 
 	/** @return		The property stored at this node, to be passed to Pre/PostEditChange. */
-	UProperty*		GetStoredProperty()		{ return NULL; }
+	UProperty*		GetStoredProperty()		{ return nullptr; }
 
 	bool GetDiffersFromDefaultForObject( FPropertyItemValueDataTrackerSlate& ValueTracker, UProperty* InProperty );
 
@@ -784,7 +791,7 @@ protected:
 
 protected:
 	/**
-	 * The node that is the parent of this node or NULL for the root
+	 * The node that is the parent of this node or nullptr for the root
 	 */
 	TWeakPtr<FPropertyNode> ParentNodeWeakPtr;
 	//@todo consolidate with ParentNodeWeakPtr, ParentNode is legacy
@@ -875,6 +882,9 @@ public:
 	virtual FComplexPropertyNode* AsComplexNode() override{ return this; }
 	virtual const FComplexPropertyNode* AsComplexNode() const override{ return this; }
 
+	virtual class FStructurePropertyNode* AsStructureNode() { return nullptr; }
+	virtual const FStructurePropertyNode* AsStructureNode() const { return nullptr; }
+
 	virtual UStruct* GetBaseStructure() = 0;
 	virtual const UStruct* GetBaseStructure() const = 0;
 
diff --git a/Engine/Source/Editor/PropertyEditor/Private/SDetailSingleItemRow.cpp b/Engine/Source/Editor/PropertyEditor/Private/SDetailSingleItemRow.cpp
--- a/Engine/Source/Editor/PropertyEditor/Private/SDetailSingleItemRow.cpp
+++ b/Engine/Source/Editor/PropertyEditor/Private/SDetailSingleItemRow.cpp
@@ -114,7 +114,14 @@ FReply SDetailSingleItemRow::OnFavoriteToggle()
 			FName CatFavName = *CategoryFavoritesName;
 			int32 SimplePropertiesNum = 0;
 			int32 AdvancePropertiesNum = 0;
-			bool HasCategoryFavorite = OwnerTreeNode.Pin()->GetDetailsView().GetCategoryInfo(CatFavName, SimplePropertiesNum, AdvancePropertiesNum);
+
+			FDetailLayoutBuilderImpl& DetailLayout = OwnerTreeNode.Pin()->GetParentCategory()->GetParentLayoutImpl();
+
+			bool HasCategoryFavorite = DetailLayout.HasCategory(CatFavName);
+			if(HasCategoryFavorite)
+			{
+				DetailLayout.DefaultCategory(CatFavName).GetCategoryInformation(SimplePropertiesNum, AdvancePropertiesNum);
+			}
 
 			//Check if the property we toggle is an advance property
 			bool IsAdvanceProperty = Customization->GetPropertyNode()->HasNodeFlags(EPropertyNodeFlags::IsAdvanced) == 0 ? false : true;
