Mesh Paint: Fixed regression with brush cursor not showing up except when in VR mode

diff --git a/Engine/Source/Editor/MeshPaint/Private/MeshPaintEdMode.cpp b/Engine/Source/Editor/MeshPaint/Private/MeshPaintEdMode.cpp
--- a/Engine/Source/Editor/MeshPaint/Private/MeshPaintEdMode.cpp
+++ b/Engine/Source/Editor/MeshPaint/Private/MeshPaintEdMode.cpp
@@ -2704,7 +2704,7 @@ void FEdModeMeshPaint::Render( const FSceneView* View, FViewport* Viewport, FPri
 		bool bIsInVRMode = false;
 		if ( IVREditorModule::IsAvailable() )
 		{
-			if ( VREditorMode != nullptr && VREditorMode->IsFullyInitialized() )
+			if ( VREditorMode != nullptr && VREditorMode->IsActive() )
 			{
 				bIsInVRMode = true;
 				for ( UViewportInteractor* Interactor : VREditorMode->GetWorldInteraction().GetInteractors() )
