Copying //UE4/Dev-Editor to //UE4/Dev-Main (Source: //UE4/Dev-Editor @ 3228984)

diff --git a/Engine/Source/Editor/MeshPaint/Private/MeshPaintEdMode.cpp b/Engine/Source/Editor/MeshPaint/Private/MeshPaintEdMode.cpp
--- a/Engine/Source/Editor/MeshPaint/Private/MeshPaintEdMode.cpp
+++ b/Engine/Source/Editor/MeshPaint/Private/MeshPaintEdMode.cpp
@@ -113,6 +113,9 @@ FEdModeMeshPaint::FEdModeMeshPaint()
 	  bPushInstanceColorsToMesh( false ),
 	  PaintingStartTime( 0.0 ),
 	  ModifiedStaticMeshes(),
+	  PaintingMeshLODIndex(0),
+	  AllowPaintingAllLOD(false),
+	  IsUserPaintAnySubLOD(false),
 	  TexturePaintingCurrentMeshComponent( NULL ),
 	  TexturePaintingStaticMeshLOD(0),
 	  PaintingTexture2D( NULL ),
@@ -177,6 +180,7 @@ void FEdModeMeshPaint::Enter()
 		bShouldUpdateTextureList = true;
 
 	}
+	PaintingMeshLODIndex = 0;
 
 	// Catch when objects are replaced when a construction script is rerun
 	GEditor->OnObjectsReplaced().AddSP(this, &FEdModeMeshPaint::OnObjectsReplaced);
@@ -197,6 +201,29 @@ void FEdModeMeshPaint::Enter()
 
 	// Create adapters for currently selected mesh components
 	CreateGeometryAdaptersForSelectedComponents();
+	//All selected staticmesh component must have the same value for bCustomOverrideVertexColorPerLOD
+	//The first item selected will rule if the AllowPaintingAllLOD will be true or false
+	TArray<UMeshComponent*> SelectedComponents = GetSelectedMeshComponents();
+	for (UMeshComponent* SelectedComponent : SelectedComponents)
+	{
+		UStaticMeshComponent *SMComponent = Cast<UStaticMeshComponent>(SelectedComponent);
+		if (SMComponent)
+		{
+			if (SelectedComponents.Num() > 1 && ComponentToAdapterMap.Num() > 1)
+			{
+				if (AllowPaintingAllLOD != SMComponent->bCustomOverrideVertexColorPerLOD)
+				{
+					SMComponent->bCustomOverrideVertexColorPerLOD = AllowPaintingAllLOD;
+					SMComponent->Modify();
+				}
+			}
+			else
+			{
+				AllowPaintingAllLOD = SMComponent->bCustomOverrideVertexColorPerLOD;
+			}
+		}
+	}
+	IsUserPaintAnySubLOD = true;
 
 	if (!Toolkit.IsValid())
 	{
@@ -227,7 +254,7 @@ void FEdModeMeshPaint::Enter()
 	//When painting vertex colors we want to force the lod level of objects being painted to LOD0.
 	if( FMeshPaintSettings::Get().ResourceType == EMeshPaintResource::VertexColors )
 	{
-		ApplyOrRemoveForceBestLOD(/*bApply=*/ true);
+		ApplyOrRemoveForceBestLOD(/*bApply=*/ true, PaintingMeshLODIndex);
 	}
 
 	TSharedPtr<FEditorWorldWrapper> EditorWorld = GEditor->GetEditorWorldManager()->GetEditorWorldWrapper( GetWorld() );
@@ -266,8 +293,11 @@ void FEdModeMeshPaint::Exit()
 	//Then stop forcing the LOD level of the mesh to LOD0.
 	if( FMeshPaintSettings::Get().ResourceType == EMeshPaintResource::VertexColors )
 	{
-		ApplyVertexColorsToAllLODs();
-		ApplyOrRemoveForceBestLOD(/*bApply=*/ false);
+		if (!AllowPaintingAllLOD)
+		{
+			ApplyVertexColorsToAllLODs();
+		}
+		ApplyOrRemoveForceBestLOD(/*bApply=*/ false, PaintingMeshLODIndex);
 	}
 
 	// The user can manipulate the editor selection lock flag in paint mode so we make sure to restore it here
@@ -1365,7 +1395,7 @@ void FEdModeMeshPaint::DoPaint( const FVector& InCameraOrigin,
 					if (UStaticMesh* StaticMesh = StaticMeshComponent->GetStaticMesh())
 					{
 						//@TODO: MESHPAINT: Direct assumptions about StaticMeshComponent
-						check(StaticMesh->GetNumLODs() > PaintingMeshLODIndex);
+						check(StaticMesh->GetNumLODs() > (int32)PaintingMeshLODIndex);
 						FStaticMeshLODResources& LODModel = StaticMesh->RenderData->LODResources[PaintingMeshLODIndex];
 
 						// Painting vertex colors
@@ -1402,8 +1432,7 @@ static bool PropagateColorsToRawMesh(UStaticMesh* StaticMesh, int32 LODIndex, FS
 	FStaticMeshRenderData& RenderData = *StaticMesh->RenderData;
 	FStaticMeshLODResources& RenderModel = RenderData.LODResources[LODIndex];
 	FColorVertexBuffer& ColorVertexBuffer = *ComponentLODInfo.OverrideVertexColors;
-
-	if (RenderData.WedgeMap.Num() > 0 && ColorVertexBuffer.GetNumVertices() == RenderModel.GetNumVertices())
+	if (LODIndex == 0 && RenderData.WedgeMap.Num() > 0 && ColorVertexBuffer.GetNumVertices() == RenderModel.GetNumVertices())
 	{
 		// Use the wedge map if it is available as it is lossless.
 		FRawMesh RawMesh;
@@ -1515,6 +1544,11 @@ void FEdModeMeshPaint::PaintMeshVertices(
 				// Mark the mesh component as modified
 				StaticMeshComponent->SetFlags(RF_Transactional);
 				StaticMeshComponent->Modify();
+				//Set the flag for when a lower mesh is paint, so we know if the user will loose some data when
+				if (PaintingMeshLODIndex > 0)
+				{
+					IsUserPaintAnySubLOD = true;
+				}
 
 				// Ensure LODData has enough entries in it, free not required.
 				StaticMeshComponent->SetLODDataCount(PaintingMeshLODIndex + 1, StaticMeshComponent->LODData.Num());
@@ -1564,7 +1598,7 @@ void FEdModeMeshPaint::PaintMeshVertices(
 			}
 			else
 			{
-				if( StaticMeshComponent->LODData.Num() > PaintingMeshLODIndex )
+				if( StaticMeshComponent->LODData.Num() > (int32)PaintingMeshLODIndex )
 				{
 					InstanceMeshLODInfo = &StaticMeshComponent->LODData[ PaintingMeshLODIndex ];
 				}
@@ -1664,20 +1698,24 @@ void FEdModeMeshPaint::PaintMeshVertices(
 			}
 			else if (Params.PaintAction == EMeshPaintAction::PushInstanceColorsToMesh)
 			{
-				InstanceMeshLODInfo = &StaticMeshComponent->LODData[ PaintingMeshLODIndex ];
-				if (InstanceMeshLODInfo->OverrideVertexColors)
+				bool SomePaintWasPropagated = false;
+				for (int32 LODIndex = 0; LODIndex < StaticMesh->RenderData->LODResources.Num(); LODIndex++)
 				{
-					// Try using the mapping generated when building the mesh.
-					if(PropagateColorsToRawMesh(StaticMesh, PaintingMeshLODIndex, *InstanceMeshLODInfo))
+					InstanceMeshLODInfo = &StaticMeshComponent->LODData[LODIndex];
+					if (InstanceMeshLODInfo->OverrideVertexColors)
 					{
-						for (int32 LODIndex = 1; LODIndex < StaticMesh->RenderData->LODResources.Num(); LODIndex++)
+						// Try using the mapping generated when building the mesh.
+						if (PropagateColorsToRawMesh(StaticMesh, LODIndex, *InstanceMeshLODInfo))
 						{
-							PropagateColorsToRawMesh(StaticMesh, LODIndex, *InstanceMeshLODInfo);
+							SomePaintWasPropagated = true;
 						}
-						RemoveComponentInstanceVertexColors(StaticMeshComponent);
-						StaticMesh->Build();
 					}
 				}
+				if (SomePaintWasPropagated)
+				{
+					RemoveComponentInstanceVertexColors(StaticMeshComponent);
+					StaticMesh->Build();
+				}
 				FEditorSupportDelegates::RedrawAllViewports.Broadcast();
 			}
 			else
@@ -1861,7 +1899,7 @@ void FEdModeMeshPaint::PaintMeshTexture( UMeshComponent* MeshComponent, const FM
 	if ((StaticMeshComponent != nullptr) && (StaticMeshComponent->GetStaticMesh() != nullptr))
 	{
 		//@TODO: Find a better way to move this generically to the adapter
-		check(StaticMeshComponent->GetStaticMesh()->GetNumLODs() > PaintingMeshLODIndex);
+		check(StaticMeshComponent->GetStaticMesh()->GetNumLODs() > (int32)PaintingMeshLODIndex);
 		FStaticMeshLODResources& LODModel = StaticMeshComponent->GetStaticMesh()->RenderData->LODResources[PaintingMeshLODIndex];
 		const int32 NumSections = LODModel.Sections.Num();
 
@@ -2983,8 +3021,11 @@ bool FEdModeMeshPaint::Select( AActor* InActor, bool bInSelected )
 				else if (FMeshPaintSettings::Get().ResourceType == EMeshPaintResource::VertexColors)
 				{
 					// Propagate painting to lower LODs and stop forcing the rendered mesh to LOD0.
-					ApplyVertexColorsToAllLODs(*GeomInfo, MeshComponent);
-					ApplyOrRemoveForceBestLOD(*GeomInfo, MeshComponent, /*bApply=*/ false);
+					if (!AllowPaintingAllLOD)
+					{
+						ApplyVertexColorsToAllLODs(*GeomInfo, MeshComponent);
+					}
+					ApplyOrRemoveForceBestLOD(*GeomInfo, MeshComponent, /*bApply=*/ false, PaintingMeshLODIndex);
 					{
 						FComponentReregisterContext ReregisterContext(MeshComponent);
 					}
@@ -3005,8 +3046,26 @@ bool FEdModeMeshPaint::Select( AActor* InActor, bool bInSelected )
 				}
 				else if (FMeshPaintSettings::Get().ResourceType == EMeshPaintResource::VertexColors)
 				{
+					//Make sure all selected component have the same value for bCustomOverrideVertexColorPerLOD
+					//The first selected item will set the AllowPaintingAllLOD value
+					UStaticMeshComponent *SMComponent = Cast<UStaticMeshComponent>(MeshComponent);
+					if (SMComponent)
+					{
+						if (MeshComponents.Num() > 1 && ComponentToAdapterMap.Num() > 1)
+						{
+							if (AllowPaintingAllLOD != SMComponent->bCustomOverrideVertexColorPerLOD)
+							{
+								SMComponent->bCustomOverrideVertexColorPerLOD = AllowPaintingAllLOD;
+								SMComponent->Modify();
+							}
+						}
+						else
+						{
+							AllowPaintingAllLOD = SMComponent->bCustomOverrideVertexColorPerLOD;
+						}
+					}
 					//Painting is done on LOD0 so force the mesh to render only LOD0.
-					ApplyOrRemoveForceBestLOD(*GeomInfo, MeshComponent, /*bApply=*/ true);
+					ApplyOrRemoveForceBestLOD(*GeomInfo, MeshComponent, /*bApply=*/ true, PaintingMeshLODIndex);
 					{
 						FComponentReregisterContext ReregisterContext(MeshComponent);
 					}
@@ -3014,7 +3073,7 @@ bool FEdModeMeshPaint::Select( AActor* InActor, bool bInSelected )
 			}
 		}
 	}
-	
+	IsUserPaintAnySubLOD = true;
 	return false;
 }
 
@@ -3650,7 +3709,7 @@ void FEdModeMeshPaint::RemoveInstanceVertexColors(UObject* Obj) const
 
 void FEdModeMeshPaint::RemoveComponentInstanceVertexColors(UStaticMeshComponent* StaticMeshComponent) const
 {
-	if( StaticMeshComponent != NULL && StaticMeshComponent->GetStaticMesh() != nullptr && StaticMeshComponent->GetStaticMesh()->GetNumLODs() > PaintingMeshLODIndex )
+	if( StaticMeshComponent != NULL && StaticMeshComponent->GetStaticMesh() != nullptr && StaticMeshComponent->GetStaticMesh()->GetNumLODs() > (int32)PaintingMeshLODIndex )
 	{
 		// Mark the mesh component as modified
 		StaticMeshComponent->Modify();
@@ -3943,7 +4002,7 @@ void FEdModeMeshPaint::FixupInstanceVertexColors() const
 	}
 }
 
-void FEdModeMeshPaint::ApplyOrRemoveForceBestLOD(bool bApply)
+void FEdModeMeshPaint::ApplyOrRemoveForceBestLOD(bool bApply, uint32 LODIndex)
 {
 	TArray<UMeshComponent*> SelectedMeshComponents = GetSelectedMeshComponents();
 
@@ -3952,18 +4011,18 @@ void FEdModeMeshPaint::ApplyOrRemoveForceBestLOD(bool bApply)
 		IMeshPaintGeometryAdapter* MeshAdapter = FindGeometryAdapter(MeshComponent);
 		if (MeshAdapter)
 		{
-			ApplyOrRemoveForceBestLOD(*MeshAdapter, MeshComponent, bApply);
+			ApplyOrRemoveForceBestLOD(*MeshAdapter, MeshComponent, bApply, LODIndex);
 		}
 	}
 }
 
-void FEdModeMeshPaint::ApplyOrRemoveForceBestLOD(const IMeshPaintGeometryAdapter& GeometryInfo, UMeshComponent* MeshComponent, bool bApply)
+void FEdModeMeshPaint::ApplyOrRemoveForceBestLOD(const IMeshPaintGeometryAdapter& GeometryInfo, UMeshComponent* MeshComponent, bool bApply, uint32 LODIndex)
 {
 	if (UStaticMeshComponent* StaticMeshComponent = Cast<UStaticMeshComponent>(MeshComponent))
 	{
 		//=0 means do not force the LOD.
 		//>0 means force the lod to x-1.
-		StaticMeshComponent->ForcedLodModel = bApply ? 1 : 0;
+		StaticMeshComponent->ForcedLodModel = bApply ? (LODIndex+1) : 0;
 	}
 }
 
@@ -3984,67 +4043,79 @@ void FEdModeMeshPaint::ApplyVertexColorsToAllLODs()
 void FEdModeMeshPaint::ApplyVertexColorsToAllLODs(const IMeshPaintGeometryAdapter& GeometryInfo, UMeshComponent* InMeshComponent)
 {
 	UStaticMeshComponent* StaticMeshComponent = Cast<UStaticMeshComponent>(InMeshComponent);
-
-	// If a static mesh component was found, apply LOD0 painting to all lower LODs.
-	if( StaticMeshComponent && StaticMeshComponent->GetStaticMesh() && FMeshPaintSettings::Get().ResourceType == EMeshPaintResource::VertexColors )
+		// If a static mesh component was found, apply LOD0 painting to all lower LODs.
+	if (!StaticMeshComponent || !StaticMeshComponent->GetStaticMesh() || FMeshPaintSettings::Get().ResourceType != EMeshPaintResource::VertexColors)
 	{
-		uint32 NumLODs = StaticMeshComponent->GetStaticMesh()->RenderData->LODResources.Num();
-		StaticMeshComponent->Modify();
+		return;
+	}
+	
+	if (StaticMeshComponent->LODData.Num() < 1)
+	{
+		//We need at least some painting on the base LOD to apply it to the lower LODs
+		return;
+	}
+	FStaticMeshComponentLODInfo& SourceCompLODInfo = StaticMeshComponent->LODData[0];
+	FStaticMeshLODResources& SourceRenderData = StaticMeshComponent->GetStaticMesh()->RenderData->LODResources[0];
+	//Make sure we have something paint in the LOD 0 to apply it to all lower LODs.
+	if (SourceCompLODInfo.OverrideVertexColors == nullptr || SourceCompLODInfo.PaintedVertices.Num() <= 0)
+	{
+		return;
+	}
 
-		// Ensure LODData has enough entries in it, free not required.
-		StaticMeshComponent->SetLODDataCount(NumLODs, StaticMeshComponent->LODData.Num());
-		for( uint32 i = 1 ; i < NumLODs ; ++i )
+	uint32 NumLODs = StaticMeshComponent->GetStaticMesh()->RenderData->LODResources.Num();
+	StaticMeshComponent->Modify();
+
+	// Ensure LODData has enough entries in it, free not required.
+	StaticMeshComponent->SetLODDataCount(NumLODs, StaticMeshComponent->LODData.Num());
+	for( uint32 i = 1 ; i < NumLODs ; ++i )
+	{
+		FStaticMeshComponentLODInfo* CurrInstanceMeshLODInfo = &StaticMeshComponent->LODData[ i ];
+		FStaticMeshLODResources& CurrRenderData = StaticMeshComponent->GetStaticMesh()->RenderData->LODResources[ i ];
+		// Destroy the instance vertex  color array if it doesn't fit
+		if(CurrInstanceMeshLODInfo->OverrideVertexColors
+			&& CurrInstanceMeshLODInfo->OverrideVertexColors->GetNumVertices() != CurrRenderData.GetNumVertices())
 		{
-			FStaticMeshComponentLODInfo* CurrInstanceMeshLODInfo = &StaticMeshComponent->LODData[ i ];
-			FStaticMeshLODResources& CurrRenderData = StaticMeshComponent->GetStaticMesh()->RenderData->LODResources[ i ];
-			// Destroy the instance vertex  color array if it doesn't fit
-			if(CurrInstanceMeshLODInfo->OverrideVertexColors
-				&& CurrInstanceMeshLODInfo->OverrideVertexColors->GetNumVertices() != CurrRenderData.GetNumVertices())
-			{
-				CurrInstanceMeshLODInfo->ReleaseOverrideVertexColorsAndBlock();
-			}	
+			CurrInstanceMeshLODInfo->ReleaseOverrideVertexColorsAndBlock();
+		}	
 
-			if(CurrInstanceMeshLODInfo->OverrideVertexColors)
-			{
-				CurrInstanceMeshLODInfo->BeginReleaseOverrideVertexColors();
-			}
-			else
-			{
-				// Setup the instance vertex color array if we don't have one yet
-				CurrInstanceMeshLODInfo->OverrideVertexColors = new FColorVertexBuffer;
-			}
+		if(CurrInstanceMeshLODInfo->OverrideVertexColors)
+		{
+			CurrInstanceMeshLODInfo->BeginReleaseOverrideVertexColors();
 		}
-			
-		FlushRenderingCommands();
-		FStaticMeshComponentLODInfo& SourceCompLODInfo = StaticMeshComponent->LODData[ 0 ];
-		FStaticMeshLODResources& SourceRenderData = StaticMeshComponent->GetStaticMesh()->RenderData->LODResources[ 0 ];
-		for( uint32 i=1 ; i < NumLODs ; ++i )
+		else
 		{
-			FStaticMeshComponentLODInfo& CurCompLODInfo = StaticMeshComponent->LODData[ i ];
-			FStaticMeshLODResources& CurRenderData = StaticMeshComponent->GetStaticMesh()->RenderData->LODResources[ i ];
-
-			check(CurCompLODInfo.OverrideVertexColors);
+			// Setup the instance vertex color array if we don't have one yet
+			CurrInstanceMeshLODInfo->OverrideVertexColors = new FColorVertexBuffer;
+		}
+	}
+			
+	FlushRenderingCommands();
+	for( uint32 i=1 ; i < NumLODs ; ++i )
+	{
+		FStaticMeshComponentLODInfo& CurCompLODInfo = StaticMeshComponent->LODData[ i ];
+		FStaticMeshLODResources& CurRenderData = StaticMeshComponent->GetStaticMesh()->RenderData->LODResources[ i ];
 
-			TArray<FColor> NewOverrideColors;
+		check(CurCompLODInfo.OverrideVertexColors);
 
-			if( SourceCompLODInfo.PaintedVertices.Num() > 0 )
-			{
-				RemapPaintedVertexColors(
-				SourceCompLODInfo.PaintedVertices,
-					*SourceCompLODInfo.OverrideVertexColors,
-					CurRenderData.PositionVertexBuffer,
-					&CurRenderData.VertexBuffer,
-					NewOverrideColors
-					);
-			}
-			if (NewOverrideColors.Num())
-			{
-				CurCompLODInfo.OverrideVertexColors->InitFromColorArray(NewOverrideColors);
-			}
+		TArray<FColor> NewOverrideColors;
 
-			// Initialize the vert. colors
-			BeginInitResource( CurCompLODInfo.OverrideVertexColors );
+		if( SourceCompLODInfo.PaintedVertices.Num() > 0 )
+		{
+			RemapPaintedVertexColors(
+			SourceCompLODInfo.PaintedVertices,
+				*SourceCompLODInfo.OverrideVertexColors,
+				CurRenderData.PositionVertexBuffer,
+				&CurRenderData.VertexBuffer,
+				NewOverrideColors
+				);
 		}
+		if (NewOverrideColors.Num())
+		{
+			CurCompLODInfo.OverrideVertexColors->InitFromColorArray(NewOverrideColors);
+		}
+
+		// Initialize the vert. colors
+		BeginInitResource( CurCompLODInfo.OverrideVertexColors );
 	}
 }
 
@@ -4135,14 +4206,14 @@ bool FEdModeMeshPaint::GetSelectedMeshInfo( int32& OutTotalBaseVertexColorBytes,
 			SelectedActor->GetComponents<UStaticMeshComponent>(StaticMeshComponents);
 			for(const auto& StaticMeshComponent : StaticMeshComponents)
 			{
-				if( StaticMeshComponent != nullptr && StaticMeshComponent->GetStaticMesh() != nullptr && StaticMeshComponent->GetStaticMesh()->GetNumLODs() > PaintingMeshLODIndex )
+				if( StaticMeshComponent != nullptr && StaticMeshComponent->GetStaticMesh() != nullptr && StaticMeshComponent->GetStaticMesh()->GetNumLODs() >(int32)PaintingMeshLODIndex )
 				{
 					// count the base mesh color data
 					FStaticMeshLODResources& LODModel = StaticMeshComponent->GetStaticMesh()->RenderData->LODResources[ PaintingMeshLODIndex ];
 					OutTotalBaseVertexColorBytes += LODModel.ColorVertexBuffer.GetNumVertices();
 
 					// count the instance color data
-					if( StaticMeshComponent->LODData.Num() > PaintingMeshLODIndex )
+					if( StaticMeshComponent->LODData.Num() > (int32)PaintingMeshLODIndex )
 					{
 						const FStaticMeshComponentLODInfo& InstanceMeshLODInfo = StaticMeshComponent->LODData[ PaintingMeshLODIndex ];
 						if( InstanceMeshLODInfo.OverrideVertexColors )
@@ -4160,6 +4231,177 @@ bool FEdModeMeshPaint::GetSelectedMeshInfo( int32& OutTotalBaseVertexColorBytes,
 	return ( NumValidMeshes > 0 );
 }
 
+uint32 FEdModeMeshPaint::GetMeshesLODMax() const
+{
+	//The maximum LOD we can paint is decide by the lowest number of LOD in the selection
+	int32 LODMin = MAX_int32;
+	for (UMeshComponent *MeshComponent : GetSelectedMeshComponents())
+	{
+		UStaticMeshComponent *SMComponent = Cast<UStaticMeshComponent>(MeshComponent);
+		if (!SMComponent || !SMComponent->GetStaticMesh())
+			continue;
+		if (LODMin > SMComponent->GetStaticMesh()->GetNumLODs() - 1)
+		{
+			LODMin = SMComponent->GetStaticMesh()->GetNumLODs() - 1;
+		}
+	}
+	if (LODMin == MAX_int32)
+	{
+		LODMin = 0;
+	}
+	return LODMin;
+}
+
+uint32 FEdModeMeshPaint::GetCurrentMeshesLOD() const
+{
+	return PaintingMeshLODIndex;
+}
+
+void FEdModeMeshPaint::SetCurrentMeshesLOD(const uint32 InLODIndex)
+{
+	//Cannot set a LOD index bigger then the max LOD
+	if (GetMeshesLODMax() < InLODIndex || InLODIndex == PaintingMeshLODIndex)
+		return;
+
+	TUniquePtr< FComponentReregisterContext > ComponentReregisterContext;
+	PaintingMeshLODIndex = InLODIndex;
+	RemoveAllGeometryAdapters();
+	CreateGeometryAdaptersForSelectedComponents();
+	ApplyOrRemoveForceBestLOD(/*bApply=*/ true, PaintingMeshLODIndex);
+	TArray<UMeshComponent*> SelectedComponents = GetSelectedMeshComponents();
+	
+	//Make sure all static mesh render is dirty since we change the force LOD
+	for (UMeshComponent* SelectedComponent : SelectedComponents)
+	{
+		UStaticMeshComponent *StaticMeshComponent = Cast<UStaticMeshComponent>(SelectedComponent);
+		if (StaticMeshComponent)
+		{
+			ComponentReregisterContext.Reset(new FComponentReregisterContext(StaticMeshComponent));
+		}
+	}
+}
+
+bool FEdModeMeshPaint::IsPaintingAllLODEnable() const
+{
+	return AllowPaintingAllLOD;
+}
+
+void FEdModeMeshPaint::SetAllowPaintingAllLOD(bool InAllowPaintingAllLOD)
+{
+	bool AbortChange = false;
+	if (!InAllowPaintingAllLOD)
+	{
+		//Only show the lost data warning if there is actually some data to loose
+		bool bShowWarningMessage = false;
+		TArray<UMeshComponent*> SelectedComponents = GetSelectedMeshComponents();
+		for (UMeshComponent* SelectedComponent : SelectedComponents)
+		{
+			UStaticMeshComponent *StaticMeshComponent = Cast<UStaticMeshComponent>(SelectedComponent);
+			if (StaticMeshComponent && StaticMeshComponent->GetStaticMesh() && FMeshPaintSettings::Get().ResourceType == EMeshPaintResource::VertexColors)
+			{
+				for (int32 LODIndex = 1; !bShowWarningMessage && LODIndex < StaticMeshComponent->LODData.Num(); ++LODIndex)
+				{
+					FStaticMeshComponentLODInfo& CompLODInfo = StaticMeshComponent->LODData[LODIndex];
+					if (CompLODInfo.OverrideVertexColors != nullptr || CompLODInfo.PaintedVertices.Num() > 0)
+					{
+						bShowWarningMessage = true;
+					}
+				}
+			}
+			if (bShowWarningMessage)
+			{
+				break;
+			}
+		}
+
+		if (bShowWarningMessage && IsUserPaintAnySubLOD)
+		{
+			//Warn the user he will loose is custom painting data
+			FSuppressableWarningDialog::FSetupInfo SetupInfo(LOCTEXT("LooseLowersLODsVertexColorsPrompt_Message", "Changing from custom LODs to base LOD only painting will propagate the base lod vertex color to all lowers LODs. This mean all lowers LODs custom vertex painting will be lost."),
+				LOCTEXT("LooseLowersLODsVertexColorsPrompt_Title", "Warning: Lowers LODs custom vertex painting will be lost!"), "Warning_LooseLowersLODsVertexColorsPrompt");
+
+			SetupInfo.ConfirmText = LOCTEXT("LooseLowersLODsVertexColorsPrompt_ConfirmText", "Continue");
+			SetupInfo.CancelText = LOCTEXT("LooseLowersLODsVertexColorsPrompt_CancelText", "Abort");
+			SetupInfo.CheckBoxText = LOCTEXT("LooseLowersLODsVertexColorsPrompt_CheckBoxText", "Always copy vertex colors without prompting");
+
+			FSuppressableWarningDialog LooseLowersLODsVertexColorsWarning(SetupInfo);
+			
+			// Prompt the user to see if they really want to propagate the base lod vert colors to the lowers LODs.
+			if (LooseLowersLODsVertexColorsWarning.ShowModal() == FSuppressableWarningDialog::Cancel)
+			{
+				AbortChange = true;
+			}
+			else
+			{
+				//Remove painting on all lowers LODs before doing the propagation
+				for (UMeshComponent* SelectedComponent : SelectedComponents)
+				{
+					UStaticMeshComponent *StaticMeshComponent = Cast<UStaticMeshComponent>(SelectedComponent);
+					if (StaticMeshComponent && StaticMeshComponent->GetStaticMesh() && FMeshPaintSettings::Get().ResourceType == EMeshPaintResource::VertexColors)
+					{
+						// Mark the mesh component as modified
+						StaticMeshComponent->Modify();
+
+						// If this is called from the Remove button being clicked the SMC wont be in a Reregister context,
+						// but when it gets called from a Paste or Copy to Source operation it's already inside a more specific
+						// SMCRecreateScene context so we shouldn't put it inside another one.
+						if (StaticMeshComponent->IsRenderStateCreated())
+						{
+							// Detach all instances of this static mesh from the scene.
+							FComponentReregisterContext ComponentReregisterContext(StaticMeshComponent);
+
+							for (int32 LODIndex = 1; LODIndex < StaticMeshComponent->LODData.Num(); ++LODIndex)
+							{
+								StaticMeshComponent->RemoveInstanceVertexColorsFromLOD(LODIndex);
+							}
+						}
+						else
+						{
+							for (int32 LODIndex = 1; LODIndex < StaticMeshComponent->LODData.Num(); ++LODIndex)
+							{
+								StaticMeshComponent->RemoveInstanceVertexColorsFromLOD(LODIndex);
+							}
+						}
+					}
+				}
+			}
+		}
+	}
+	
+	//The user cancel the change, avoid changing the value
+	if (AbortChange)
+	{
+		return;
+	}
+	
+	//Reset the flag that tell us if there was a change in the lower LODs
+	IsUserPaintAnySubLOD = false;
+	
+	AllowPaintingAllLOD = InAllowPaintingAllLOD;
+
+	TArray<UMeshComponent*> SelectedComponents = GetSelectedMeshComponents();
+	for (UMeshComponent* SelectedComponent : SelectedComponents)
+	{
+		UStaticMeshComponent *StaticMeshComponent = Cast<UStaticMeshComponent>(SelectedComponent);
+		if (StaticMeshComponent)
+		{
+			StaticMeshComponent->Modify();
+			StaticMeshComponent->bCustomOverrideVertexColorPerLOD = InAllowPaintingAllLOD;
+		}
+	}
+	//Always propagate the base LOD when we switch the painting to all LODs
+	//In case we go from false to true, we want to tweak the propagation of the base LOD
+	//In case we go from true to false, we want to force a propagate to be sure the custom painting is done over what the user already paint
+	//In case the staticmesh has some data comming from fbx, the propagate will not do anything
+	//if there is nothing that was overrite in the base LOD. So the user will be able to paint
+	//over the painting done in a DCC.
+	ApplyVertexColorsToAllLODs();
+	if (!AllowPaintingAllLOD)
+	{
+		SetCurrentMeshesLOD(0);
+	}
+}
+
 void FEdModeMeshPaint::SetBrushRadiiDefault( float InBrushRadius )
 {
 	float MinBrushRadius, MaxBrushRadius;
diff --git a/Engine/Source/Editor/MeshPaint/Private/SMeshPaint.cpp b/Engine/Source/Editor/MeshPaint/Private/SMeshPaint.cpp
--- a/Engine/Source/Editor/MeshPaint/Private/SMeshPaint.cpp
+++ b/Engine/Source/Editor/MeshPaint/Private/SMeshPaint.cpp
@@ -2107,6 +2107,50 @@ void SMeshPaint::Construct(const FArguments& InArgs, TSharedRef<FMeshPaintToolKi
 								.OnSelectionChanged(this, &SMeshPaint::OnVertexPaintColorViewModeChanged)
 							]
 						]
+						+SVerticalBox::Slot()
+						.AutoHeight()
+						.Padding(StandardPadding)
+						[
+							SNew(SHorizontalBox)
+							.Visibility(this, &SMeshPaint::GetResourceTypeVerticesVisibility )
+							+ SHorizontalBox::Slot()
+							.Padding(2.0f, 0.0f)
+							.AutoWidth()
+							.HAlign(HAlign_Left)
+							.VAlign(VAlign_Center)
+							[
+								SNew(SCheckBox)
+								.IsChecked(this, &SMeshPaint::IsPaintingAllLOD)
+								.OnCheckStateChanged(this, &SMeshPaint::OnAllowPaintingAllLODChanged)
+							]
+							+SHorizontalBox::Slot() 
+							.Padding(2.0f, 0.0f) 
+							.FillWidth(1) 
+							.HAlign(HAlign_Left)
+							.VAlign(VAlign_Center)
+							[
+								SNew(STextBlock)
+								.Text(LOCTEXT("MeshPaint_VertexColorLODSelectionLabel", "Allow Paint On All LOD"))
+								.ToolTipText(LOCTEXT("MeshPaint_VertexColorLODSelectionLabelTooltip", "When uncheck, the painting on the base LOD will be propagate automatically to all other LODs when exiting the mode or changing the selection."))
+							]
+							+SHorizontalBox::Slot()
+							.AutoWidth()
+							.Padding(0.0f, 0.0f, 2.0f, 0.0f) 
+							.HAlign(HAlign_Right)
+							.VAlign(VAlign_Center)
+							[
+								SNew(SNumericEntryBox<uint32>)
+								.IsEnabled(this, &SMeshPaint::IsPaintingAllLODEnable)
+								.MinDesiredValueWidth(100)
+								.AllowSpin(true)
+								.MinSliderValue(0)
+								.MaxSliderValue(this, &SMeshPaint::GetMeshesLODMax)
+								.MinValue(0)
+								.MaxValue(this, &SMeshPaint::GetMeshesLODMax)
+								.Value(this, &SMeshPaint::GetCurrentMeshesLOD)
+								.OnValueChanged(this, &SMeshPaint::SetCurrentMeshesLOD)
+							]
+						]
 					]
 				]
 			]
@@ -2152,6 +2196,36 @@ EVisibility SMeshPaint::GetImportVertexColorsVisibility() const
 		: EVisibility::Collapsed;
 }
 
+TOptional<uint32> SMeshPaint::GetMeshesLODMax() const
+{
+	return MeshPaintEditMode->GetMeshesLODMax();
+}
+
+TOptional<uint32> SMeshPaint::GetCurrentMeshesLOD() const
+{
+	return MeshPaintEditMode->GetCurrentMeshesLOD();
+}
+
+void SMeshPaint::SetCurrentMeshesLOD(const uint32 InLODIndex)
+{
+	MeshPaintEditMode->SetCurrentMeshesLOD(InLODIndex);
+}
+
+ECheckBoxState SMeshPaint::IsPaintingAllLOD() const
+{
+	return IsPaintingAllLODEnable() ? ECheckBoxState::Checked : ECheckBoxState::Unchecked;
+}
+
+bool SMeshPaint::IsPaintingAllLODEnable() const
+{
+	return MeshPaintEditMode->IsPaintingAllLODEnable();
+}
+
+void SMeshPaint::OnAllowPaintingAllLODChanged(ECheckBoxState InCheckState)
+{
+	MeshPaintEditMode->SetAllowPaintingAllLOD(InCheckState == ECheckBoxState::Checked);
+}
+
 TOptional<float> SMeshPaint::GetBrushRadius() const
 {
 	return MeshPaintEditMode->GetBrushRadiiDefault();
diff --git a/Engine/Source/Editor/MeshPaint/Private/SMeshPaint.h b/Engine/Source/Editor/MeshPaint/Private/SMeshPaint.h
--- a/Engine/Source/Editor/MeshPaint/Private/SMeshPaint.h
+++ b/Engine/Source/Editor/MeshPaint/Private/SMeshPaint.h
@@ -106,6 +106,24 @@ protected:
 	/** Returns the visibility state of the button that imports vertex colors from TGA */
 	EVisibility GetImportVertexColorsVisibility() const;
 
+	/** Returns the maximum LOD index for the mesh we paint*/
+	TOptional<uint32> GetMeshesLODMax() const;
+
+	/** Returns the current LOD index we paint*/
+	TOptional<uint32> GetCurrentMeshesLOD() const;
+
+	/** Call when the user change the current LOD index we paint */
+	void SetCurrentMeshesLOD(const uint32 InLODIndex);
+
+	/** Return if we paint on all LOD or just on the base LOD */
+	ECheckBoxState IsPaintingAllLOD() const;
+
+	/** Return if we paint on all LOD or just on the base LOD */
+	bool IsPaintingAllLODEnable() const;
+
+	/** Call when the user change the AllowPaintingAllLOD option */
+	void OnAllowPaintingAllLODChanged(ECheckBoxState InCheckState);
+
 	/** Returns the value for the brush radius setting */
 	TOptional<float> GetBrushRadius() const;
 
diff --git a/Engine/Source/Editor/MeshPaint/Public/MeshPaintEdMode.h b/Engine/Source/Editor/MeshPaint/Public/MeshPaintEdMode.h
--- a/Engine/Source/Editor/MeshPaint/Public/MeshPaintEdMode.h
+++ b/Engine/Source/Editor/MeshPaint/Public/MeshPaintEdMode.h
@@ -452,10 +452,10 @@ public:
 	void ApplyVertexColorsToAllLODs();
 
 	/** Forces all selected meshes to their best LOD, or removes the forcing. */
-	void ApplyOrRemoveForceBestLOD(bool bApply);
+	void ApplyOrRemoveForceBestLOD(bool bApply, uint32 LODIndex);
 
 	/** Forces the static mesh component to it's best LOD, or removes the forcing. */
-	void ApplyOrRemoveForceBestLOD(const IMeshPaintGeometryAdapter& GeometryInfo, UMeshComponent* InMeshComponent, bool bApply);
+	void ApplyOrRemoveForceBestLOD(const IMeshPaintGeometryAdapter& GeometryInfo, UMeshComponent* InMeshComponent, bool bApply, uint32 LODIndex);
 
 	/** Applies vertex color painting on LOD0 to all lower LODs. */
 	void ApplyVertexColorsToAllLODs(const IMeshPaintGeometryAdapter& GeometryInfo, UMeshComponent* InMeshComponent);
@@ -475,6 +475,21 @@ public:
 	/** Returns information about the currently selected mesh */
 	bool GetSelectedMeshInfo( int32& OutTotalBaseVertexColorBytes, int32& OutTotalInstanceVertexColorBytes, bool& bOutHasInstanceMaterialAndTexture ) const;
 
+	/** Returns the maximum LOD index for the mesh we paint*/
+	uint32 GetMeshesLODMax() const;
+
+	/** Returns the current LOD index we paint*/
+	uint32 GetCurrentMeshesLOD() const;
+
+	/** Call when the user change the current LOD index we paint */
+	void SetCurrentMeshesLOD(const uint32 InLODIndex);
+
+	/** Return if we paint on all LOD or just on the base LOD */
+	bool IsPaintingAllLODEnable() const;
+
+	/** Call when the user change the AllowPaintingAllLOD option */
+	void SetAllowPaintingAllLOD(bool InAllowPaintingAllLOD);
+
 	/** Sets the default brush radius size */
 	void SetBrushRadiiDefault( float InBrushRadius );
 
@@ -778,7 +793,11 @@ private:
 
 	/** Which mesh LOD index we're painting */
 	// @todo MeshPaint: Allow painting on other LODs?
-	static const uint32 PaintingMeshLODIndex = 0;
+	uint32 PaintingMeshLODIndex;
+
+	bool AllowPaintingAllLOD;
+
+	bool IsUserPaintAnySubLOD;
 
 	/** Texture paint: The mesh components that we're currently painting */
 	UMeshComponent* TexturePaintingCurrentMeshComponent;
diff --git a/Engine/Source/Editor/Persona/Persona.Build.cs b/Engine/Source/Editor/Persona/Persona.Build.cs
--- a/Engine/Source/Editor/Persona/Persona.Build.cs
+++ b/Engine/Source/Editor/Persona/Persona.Build.cs
@@ -39,7 +39,9 @@ public class Persona : ModuleRules
 				"EditorWidgets",
                 "BlueprintGraph",
                 "RHI",
-			}
+                "Json",
+                "JsonUtilities",
+            }
 		);
 
 		DynamicallyLoadedModuleNames.AddRange(
diff --git a/Engine/Source/Editor/Persona/Private/PersonaMeshDetails.cpp b/Engine/Source/Editor/Persona/Private/PersonaMeshDetails.cpp
--- a/Engine/Source/Editor/Persona/Private/PersonaMeshDetails.cpp
+++ b/Engine/Source/Editor/Persona/Private/PersonaMeshDetails.cpp
@@ -52,6 +52,10 @@
 #include "Animation/AnimInstance.h"
 #include "IPersonaToolkit.h"
 #include "Widgets/Input/SNumericEntryBox.h"
+#include "IPersonaPreviewScene.h"
+#include "IDocumentation.h"
+#include "JsonObjectConverter.h"
+#include "SWrapBox.h"
 
 #define LOCTEXT_NAMESPACE "PersonaMeshDetails"
 
@@ -202,38 +206,61 @@ void SSkeletalLODActions::Construct(const FArguments& InArgs)
 	this->ChildSlot
 
 		[
-			SNew(SHorizontalBox)
-			+ SHorizontalBox::Slot()
-			.AutoWidth()
+			SNew(SVerticalBox)
+			+ SVerticalBox::Slot()
 			[
-				SNew(SButton)
-				.OnClicked(OnRemoveLODClicked)
+				SNew(SWrapBox)
+				.UseAllottedWidth(true)
+				+ SWrapBox::Slot()
+				.Padding(FMargin(0, 0, 2, 4))
 				[
-					SNew(STextBlock)
-					.Text(LOCTEXT("RemoveLOD", "Remove this LOD"))
+					SNew(SBox)
+					.WidthOverride(120.f)
+					[
+						SNew(SButton)
+						.HAlign(HAlign_Center)
+						.VAlign(VAlign_Center)
+						.OnClicked(OnRemoveLODClicked)
+						[
+							SNew(STextBlock)
+							.Text(LOCTEXT("RemoveLOD", "Remove this LOD"))
+						]
+					]
 				]
-			]
-			+ SHorizontalBox::Slot()
-			.AutoWidth()
-			[
-				SNew(SButton)
-				.ToolTipText(this, &SSkeletalLODActions::GetReimportButtonToolTipText)
-				.IsEnabled(this, &SSkeletalLODActions::CanReimportFromSource)
-				.OnClicked(OnReimportClicked)
+				+ SWrapBox::Slot()
+				.Padding(FMargin(0, 0, 2, 4))
 				[
-					SNew(STextBlock)
-					.Text(LOCTEXT("ReimportLOD", "Reimport"))
+					SNew(SBox)
+					.WidthOverride(120.f)
+					[
+						SNew(SButton)
+						.HAlign(HAlign_Center)
+						.VAlign(VAlign_Center)
+						.ToolTipText(this, &SSkeletalLODActions::GetReimportButtonToolTipText)
+						.IsEnabled(this, &SSkeletalLODActions::CanReimportFromSource)
+						.OnClicked(OnReimportClicked)
+						[
+							SNew(STextBlock)
+							.Text(LOCTEXT("ReimportLOD", "Reimport"))
+						]
+					]
 				]
-			]
-			+ SHorizontalBox::Slot()
-			.AutoWidth()
-			[
-				SNew(SButton)
-				.ToolTipText(this, &SSkeletalLODActions::GetReimportButtonNewFileToolTipText)
-				.OnClicked(OnReimportNewFileClicked)
+				+ SWrapBox::Slot()
+				.Padding(FMargin(0, 0, 2, 4))
 				[
-					SNew(STextBlock)
-					.Text(LOCTEXT("ReimportLOD_NewFile", "Reimport (New File)"))
+					SNew(SBox)
+					.WidthOverride(120.f)
+					[
+						SNew(SButton)
+						.HAlign(HAlign_Center)
+						.VAlign(VAlign_Center)
+						.ToolTipText(this, &SSkeletalLODActions::GetReimportButtonNewFileToolTipText)
+						.OnClicked(OnReimportNewFileClicked)
+						[
+							SNew(STextBlock)
+							.Text(LOCTEXT("ReimportLOD_NewFile", "Reimport (New File)"))
+						]
+					]
 				]
 			]
 		];
@@ -629,11 +656,359 @@ void FSkelMeshReductionSettingsLayout::OnSkinningImportanceChanged(TSharedPtr<FS
 /**
 * FPersonaMeshDetails
 */
+FPersonaMeshDetails::~FPersonaMeshDetails()
+{
+	if (HasValidPersonaToolkit())
+	{
+		TSharedRef<IPersonaPreviewScene> PreviewScene = GetPersonaToolkit()->GetPreviewScene();
+		PreviewScene->UnregisterOnPreviewMeshChanged(this);
+	}
+}
+
 TSharedRef<IDetailCustomization> FPersonaMeshDetails::MakeInstance(TSharedRef<class IPersonaToolkit> InPersonaToolkit)
 {
 	return MakeShareable( new FPersonaMeshDetails(InPersonaToolkit) );
 }
 
+void FPersonaMeshDetails::OnCopySectionList(int32 LODIndex)
+{
+	USkeletalMesh* Mesh = GetPersonaToolkit()->GetMesh();
+
+	if (Mesh != nullptr)
+	{
+		FSkeletalMeshResource* ImportedResource = Mesh->GetImportedResource();
+
+		if (ImportedResource != nullptr && ImportedResource->LODModels.IsValidIndex(LODIndex))
+		{
+			const FStaticLODModel& Model = ImportedResource->LODModels[LODIndex];
+
+			TSharedRef<FJsonObject> RootJsonObject = MakeShareable(new FJsonObject());
+
+			for (int32 SectionIdx = 0; SectionIdx < Model.Sections.Num(); ++SectionIdx)
+			{
+				const FSkelMeshSection& ModelSection = Model.Sections[SectionIdx];
+
+				TSharedPtr<FJsonObject> JSonSection = MakeShareable(new FJsonObject);
+
+				JSonSection->SetNumberField(TEXT("MaterialIndex"), ModelSection.MaterialIndex);
+				JSonSection->SetBoolField(TEXT("RecomputeTangent"), ModelSection.bRecomputeTangent);
+				JSonSection->SetBoolField(TEXT("CastShadow"), ModelSection.bCastShadow);
+
+				RootJsonObject->SetObjectField(FString::Printf(TEXT("Section_%d"), SectionIdx), JSonSection);
+			}
+
+			typedef TJsonWriter<TCHAR, TPrettyJsonPrintPolicy<TCHAR>> FStringWriter;
+			typedef TJsonWriterFactory<TCHAR, TPrettyJsonPrintPolicy<TCHAR>> FStringWriterFactory;
+
+			FString CopyStr;
+			TSharedRef<FStringWriter> Writer = FStringWriterFactory::Create(&CopyStr);
+			FJsonSerializer::Serialize(RootJsonObject, Writer);
+
+			if (!CopyStr.IsEmpty())
+			{
+				FPlatformMisc::ClipboardCopy(*CopyStr);
+			}
+		}
+	}
+}
+
+bool FPersonaMeshDetails::OnCanCopySectionList(int32 LODIndex) const
+{
+	USkeletalMesh* Mesh = GetPersonaToolkit()->GetMesh();
+
+	if (Mesh != nullptr)
+	{
+		FSkeletalMeshResource* ImportedResource = Mesh->GetImportedResource();
+
+		if (ImportedResource != nullptr && ImportedResource->LODModels.IsValidIndex(LODIndex))
+		{
+			return ImportedResource->LODModels[LODIndex].Sections.Num() > 0;
+		}
+	}
+
+	return false;
+}
+
+void FPersonaMeshDetails::OnPasteSectionList(int32 LODIndex)
+{
+	USkeletalMesh* Mesh = GetPersonaToolkit()->GetMesh();
+
+	if (Mesh != nullptr)
+	{
+		FString PastedText;
+		FPlatformMisc::ClipboardPaste(PastedText);
+
+		TSharedPtr<FJsonObject> RootJsonObject;
+		TSharedRef<TJsonReader<>> Reader = TJsonReaderFactory<>::Create(PastedText);
+		FJsonSerializer::Deserialize(Reader, RootJsonObject);
+
+		if (RootJsonObject.IsValid())
+		{
+			FSkeletalMeshResource* ImportedResource = Mesh->GetImportedResource();
+
+			if (ImportedResource != nullptr && ImportedResource->LODModels.IsValidIndex(LODIndex))
+			{
+				FScopedTransaction Transaction(LOCTEXT("PersonaChangedPasteSectionList", "Pasted section list"));
+				Mesh->Modify();
+
+				FStaticLODModel& Model = ImportedResource->LODModels[LODIndex];
+
+				for (int32 SectionIdx = 0; SectionIdx < Model.Sections.Num(); ++SectionIdx)
+				{
+					FSkelMeshSection& ModelSection = Model.Sections[SectionIdx];
+
+					const TSharedPtr<FJsonObject>* JSonSection = nullptr;
+					if (RootJsonObject->TryGetObjectField(FString::Printf(TEXT("Section_%d"), SectionIdx), JSonSection))
+					{
+						int32 Value;
+						if ((*JSonSection)->TryGetNumberField(TEXT("MaterialIndex"), Value))
+						{
+							ModelSection.MaterialIndex = (uint16)Value;
+						}
+
+						(*JSonSection)->TryGetBoolField(TEXT("RecomputeTangent"), ModelSection.bRecomputeTangent);
+						(*JSonSection)->TryGetBoolField(TEXT("CastShadow"), ModelSection.bCastShadow);
+					}
+				}
+
+				Mesh->PostEditChange();
+			}
+		}
+	}
+}
+
+void FPersonaMeshDetails::OnCopySectionItem(int32 LODIndex, int32 SectionIndex)
+{
+	USkeletalMesh* Mesh = GetPersonaToolkit()->GetMesh();
+
+	if (Mesh != nullptr)
+	{
+		FSkeletalMeshResource* ImportedResource = Mesh->GetImportedResource();
+
+		if (ImportedResource != nullptr && ImportedResource->LODModels.IsValidIndex(LODIndex))
+		{
+			const FStaticLODModel& Model = ImportedResource->LODModels[LODIndex];
+
+			TSharedRef<FJsonObject> RootJsonObject = MakeShareable(new FJsonObject());
+
+			if (Model.Sections.IsValidIndex(SectionIndex))
+			{
+				const FSkelMeshSection& ModelSection = Model.Sections[SectionIndex];
+
+				RootJsonObject->SetNumberField(TEXT("MaterialIndex"), ModelSection.MaterialIndex);
+				RootJsonObject->SetBoolField(TEXT("RecomputeTangent"), ModelSection.bRecomputeTangent);
+				RootJsonObject->SetBoolField(TEXT("CastShadow"), ModelSection.bCastShadow);
+			}
+
+			typedef TJsonWriter<TCHAR, TPrettyJsonPrintPolicy<TCHAR>> FStringWriter;
+			typedef TJsonWriterFactory<TCHAR, TPrettyJsonPrintPolicy<TCHAR>> FStringWriterFactory;
+
+			FString CopyStr;
+			TSharedRef<FStringWriter> Writer = FStringWriterFactory::Create(&CopyStr);
+			FJsonSerializer::Serialize(RootJsonObject, Writer);
+
+			if (!CopyStr.IsEmpty())
+			{
+				FPlatformMisc::ClipboardCopy(*CopyStr);
+			}
+		}
+	}
+}
+
+bool FPersonaMeshDetails::OnCanCopySectionItem(int32 LODIndex, int32 SectionIndex) const
+{
+	USkeletalMesh* Mesh = GetPersonaToolkit()->GetMesh();
+
+	if (Mesh != nullptr)
+	{
+		FSkeletalMeshResource* ImportedResource = Mesh->GetImportedResource();
+
+		if (ImportedResource != nullptr && ImportedResource->LODModels.IsValidIndex(LODIndex))
+		{
+			return ImportedResource->LODModels[LODIndex].Sections.IsValidIndex(SectionIndex);
+		}
+	}
+
+	return false;
+}
+
+void FPersonaMeshDetails::OnPasteSectionItem(int32 LODIndex, int32 SectionIndex)
+{
+	USkeletalMesh* Mesh = GetPersonaToolkit()->GetMesh();
+
+	if (Mesh != nullptr)
+	{
+		FString PastedText;
+		FPlatformMisc::ClipboardPaste(PastedText);
+
+		TSharedPtr<FJsonObject> RootJsonObject;
+		TSharedRef<TJsonReader<>> Reader = TJsonReaderFactory<>::Create(PastedText);
+		FJsonSerializer::Deserialize(Reader, RootJsonObject);
+
+		if (RootJsonObject.IsValid())
+		{
+			FSkeletalMeshResource* ImportedResource = Mesh->GetImportedResource();
+
+			if (ImportedResource != nullptr && ImportedResource->LODModels.IsValidIndex(LODIndex))
+			{
+				FStaticLODModel& Model = ImportedResource->LODModels[LODIndex];
+
+				FScopedTransaction Transaction(LOCTEXT("PersonaChangedPasteSectionItem", "Pasted section item"));
+				Mesh->Modify();
+
+				if (Model.Sections.IsValidIndex(SectionIndex))
+				{
+					FSkelMeshSection& ModelSection = Model.Sections[SectionIndex];
+
+					int32 Value;
+					if (RootJsonObject->TryGetNumberField(TEXT("MaterialIndex"), Value))
+					{
+						ModelSection.MaterialIndex = (uint16)Value;
+					}
+
+					RootJsonObject->TryGetBoolField(TEXT("RecomputeTangent"), ModelSection.bRecomputeTangent);
+					RootJsonObject->TryGetBoolField(TEXT("CastShadow"), ModelSection.bCastShadow);
+				}
+
+				Mesh->PostEditChange();
+			}
+		}
+	}
+}
+
+void FPersonaMeshDetails::OnCopyMaterialList()
+{
+	USkeletalMesh* Mesh = GetPersonaToolkit()->GetMesh();
+
+	if (Mesh != nullptr)
+	{
+		UProperty* Property = USkeletalMesh::StaticClass()->FindPropertyByName(GET_MEMBER_NAME_STRING_CHECKED(USkeletalMesh, Materials));
+		auto JsonValue = FJsonObjectConverter::UPropertyToJsonValue(Property, &Mesh->Materials, 0, 0);
+
+		typedef TJsonWriter<TCHAR, TPrettyJsonPrintPolicy<TCHAR>> FStringWriter;
+		typedef TJsonWriterFactory<TCHAR, TPrettyJsonPrintPolicy<TCHAR>> FStringWriterFactory;
+
+		FString CopyStr;
+		TSharedRef<FStringWriter> Writer = FStringWriterFactory::Create(&CopyStr);
+		FJsonSerializer::Serialize(JsonValue.ToSharedRef(), TEXT(""), Writer);
+
+		if (!CopyStr.IsEmpty())
+		{
+			FPlatformMisc::ClipboardCopy(*CopyStr);
+		}
+	}
+}
+
+bool FPersonaMeshDetails::OnCanCopyMaterialList() const
+{
+	USkeletalMesh* Mesh = GetPersonaToolkit()->GetMesh();
+
+	if (Mesh != nullptr)
+	{
+		return Mesh->Materials.Num() > 0;
+	}
+
+	return false;
+}
+
+void FPersonaMeshDetails::OnPasteMaterialList()
+{
+	USkeletalMesh* Mesh = GetPersonaToolkit()->GetMesh();
+
+	if (Mesh != nullptr)
+	{
+		FString PastedText;
+		FPlatformMisc::ClipboardPaste(PastedText);
+
+		TSharedPtr<FJsonValue> RootJsonValue;
+		TSharedRef<TJsonReader<>> Reader = TJsonReaderFactory<>::Create(PastedText);
+		FJsonSerializer::Deserialize(Reader, RootJsonValue);
+
+		if (RootJsonValue.IsValid())
+		{
+			UProperty* Property = USkeletalMesh::StaticClass()->FindPropertyByName(GET_MEMBER_NAME_STRING_CHECKED(USkeletalMesh, Materials));
+
+			Mesh->PreEditChange(Property);
+			FScopedTransaction Transaction(LOCTEXT("PersonaChangedPasteMaterialList", "Pasted material list"));
+			Mesh->Modify();
+
+			FJsonObjectConverter::JsonValueToUProperty(RootJsonValue, Property, &Mesh->Materials, 0, 0);
+
+			Mesh->PostEditChange();
+		}
+	}
+}
+
+void FPersonaMeshDetails::OnCopyMaterialItem(int32 CurrentSlot)
+{
+	USkeletalMesh* Mesh = GetPersonaToolkit()->GetMesh();
+
+	if (Mesh != nullptr)
+	{
+		TSharedRef<FJsonObject> RootJsonObject = MakeShareable(new FJsonObject());
+
+		if (Mesh->Materials.IsValidIndex(CurrentSlot))
+		{
+			const FSkeletalMaterial& Material = Mesh->Materials[CurrentSlot];
+
+			FJsonObjectConverter::UStructToJsonObject(FSkeletalMaterial::StaticStruct(), &Material, RootJsonObject, 0, 0);
+		}
+
+		typedef TJsonWriter<TCHAR, TPrettyJsonPrintPolicy<TCHAR>> FStringWriter;
+		typedef TJsonWriterFactory<TCHAR, TPrettyJsonPrintPolicy<TCHAR>> FStringWriterFactory;
+
+		FString CopyStr;
+		TSharedRef<FStringWriter> Writer = FStringWriterFactory::Create(&CopyStr);
+		FJsonSerializer::Serialize(RootJsonObject, Writer);
+
+		if (!CopyStr.IsEmpty())
+		{
+			FPlatformMisc::ClipboardCopy(*CopyStr);
+		}
+	}
+}
+
+bool FPersonaMeshDetails::OnCanCopyMaterialItem(int32 CurrentSlot) const
+{
+	USkeletalMesh* Mesh = GetPersonaToolkit()->GetMesh();
+
+	if (Mesh != nullptr)
+	{
+		return Mesh->Materials.IsValidIndex(CurrentSlot);
+	}
+
+	return false;
+}
+
+void FPersonaMeshDetails::OnPasteMaterialItem(int32 CurrentSlot)
+{
+	USkeletalMesh* Mesh = GetPersonaToolkit()->GetMesh();
+
+	if (Mesh != nullptr)
+	{
+		FString PastedText;
+		FPlatformMisc::ClipboardPaste(PastedText);
+
+		TSharedPtr<FJsonObject> RootJsonObject;
+		TSharedRef<TJsonReader<>> Reader = TJsonReaderFactory<>::Create(PastedText);
+		FJsonSerializer::Deserialize(Reader, RootJsonObject);
+
+		if (RootJsonObject.IsValid())
+		{
+			Mesh->PreEditChange(USkeletalMesh::StaticClass()->FindPropertyByName(GET_MEMBER_NAME_STRING_CHECKED(USkeletalMesh, Materials)));
+			FScopedTransaction Transaction(LOCTEXT("PersonaChangedPasteMaterialList", "Pasted material list"));
+			Mesh->Modify();
+
+			if (Mesh->Materials.IsValidIndex(CurrentSlot))
+			{
+				FJsonObjectConverter::JsonObjectToUStruct(RootJsonObject.ToSharedRef(), FSkeletalMaterial::StaticStruct(), &Mesh->Materials[CurrentSlot], 0, 0);
+			}
+
+			Mesh->PostEditChange();
+		}
+	}
+}
+
 void FPersonaMeshDetails::AddLODLevelCategories(IDetailLayoutBuilder& DetailLayout)
 {
 	USkeletalMesh* SkelMesh = GetPersonaToolkit()->GetMesh();
@@ -652,6 +1027,8 @@ void FPersonaMeshDetails::AddLODLevelCategories(IDetailLayoutBuilder& DetailLayo
 			FString MaterialCategoryName = FString(TEXT("Materials"));
 			IDetailCategoryBuilder& MaterialCategory = DetailLayout.EditCategory(*MaterialCategoryName, FText::GetEmpty(), ECategoryPriority::Important);
 			MaterialCategory.AddCustomRow(LOCTEXT("AddLODLevelCategories_MaterialArrayOperationAdd", "Materials Operation Add Material Slot"))
+				.CopyAction(FUIAction(FExecuteAction::CreateSP(this, &FPersonaMeshDetails::OnCopyMaterialList), FCanExecuteAction::CreateSP(this, &FPersonaMeshDetails::OnCanCopyMaterialList)))
+				.PasteAction(FUIAction(FExecuteAction::CreateSP(this, &FPersonaMeshDetails::OnPasteMaterialList)))
 				.NameContent()
 				.HAlign(HAlign_Left)
 				.VAlign(VAlign_Center)
@@ -707,6 +1084,11 @@ void FPersonaMeshDetails::AddLODLevelCategories(IDetailLayoutBuilder& DetailLayo
 				MaterialListDelegates.OnMaterialChanged.BindSP(this, &FPersonaMeshDetails::OnMaterialArrayChanged, 0);
 				MaterialListDelegates.OnGenerateCustomMaterialWidgets.BindSP(this, &FPersonaMeshDetails::OnGenerateCustomMaterialWidgetsForMaterialArray, 0);
 				MaterialListDelegates.OnMaterialListDirty.BindSP(this, &FPersonaMeshDetails::OnMaterialListDirty);
+
+				MaterialListDelegates.OnCopyMaterialItem.BindSP(this, &FPersonaMeshDetails::OnCopyMaterialItem);
+				MaterialListDelegates.OnCanCopyMaterialItem.BindSP(this, &FPersonaMeshDetails::OnCanCopyMaterialItem);
+				MaterialListDelegates.OnPasteMaterialItem.BindSP(this, &FPersonaMeshDetails::OnPasteMaterialItem);
+
 				MaterialCategory.AddCustomBuilder(MakeShareable(new FMaterialList(MaterialCategory.GetParentLayout(), MaterialListDelegates, false)));
 			}
 		}
@@ -752,9 +1134,8 @@ void FPersonaMeshDetails::AddLODLevelCategories(IDetailLayoutBuilder& DetailLayo
 					.Font(IDetailLayoutBuilder::GetDetailFontItalic())
 				];
 
-			// want to make sure if this data has imporetd or not
+			// want to make sure if this data has imported or not
 			LODCategory.HeaderContent(LODCategoryWidget);
-
 			{
 				FSectionListDelegates SectionListDelegates;
 
@@ -762,6 +1143,14 @@ void FPersonaMeshDetails::AddLODLevelCategories(IDetailLayoutBuilder& DetailLayo
 				SectionListDelegates.OnSectionChanged.BindSP(this, &FPersonaMeshDetails::OnSectionChanged);
 				SectionListDelegates.OnGenerateCustomNameWidgets.BindSP(this, &FPersonaMeshDetails::OnGenerateCustomNameWidgetsForSection);
 				SectionListDelegates.OnGenerateCustomSectionWidgets.BindSP(this, &FPersonaMeshDetails::OnGenerateCustomSectionWidgetsForSection);
+
+				SectionListDelegates.OnCopySectionList.BindSP(this, &FPersonaMeshDetails::OnCopySectionList, LODIndex);
+				SectionListDelegates.OnCanCopySectionList.BindSP(this, &FPersonaMeshDetails::OnCanCopySectionList, LODIndex);
+				SectionListDelegates.OnPasteSectionList.BindSP(this, &FPersonaMeshDetails::OnPasteSectionList, LODIndex);
+				SectionListDelegates.OnCopySectionItem.BindSP(this, &FPersonaMeshDetails::OnCopySectionItem);
+				SectionListDelegates.OnCanCopySectionItem.BindSP(this, &FPersonaMeshDetails::OnCanCopySectionItem);
+				SectionListDelegates.OnPasteSectionItem.BindSP(this, &FPersonaMeshDetails::OnPasteSectionItem);
+
 				LODCategory.AddCustomBuilder(MakeShareable(new FSectionList(LODCategory.GetParentLayout(), SectionListDelegates, (LODIndex > 0))));
 			}
 			// add each LOD Info to each LOD category
@@ -928,13 +1317,26 @@ FReply FPersonaMeshDetails::RemoveOneLOD(int32 LODIndex)
 	check(SkelMesh);
 	check(SkelMesh->LODInfo.IsValidIndex(LODIndex));
 
-	FSkeletalMeshUpdateContext UpdateContext;
-	UpdateContext.SkeletalMesh = SkelMesh;
-	UpdateContext.AssociatedComponents.Push(GetPersonaToolkit()->GetPreviewMeshComponent());
+	if (LODIndex > 0)
+	{
+		FText ConfirmRemoveLODText = FText::Format( LOCTEXT("PersonaRemoveLOD_Confirmation", "Areyou sure you want to remove LOD {0} from {1}?"), LODIndex, FText::FromString(SkelMesh->GetName()) );
 
-	FLODUtilities::RemoveLOD(UpdateContext, LODIndex);
+		if ( FMessageDialog::Open(EAppMsgType::YesNo, ConfirmRemoveLODText) == EAppReturnType::Yes )
+		{
+			FText RemoveLODText = FText::Format( LOCTEXT("OnPersonaRemoveLOD", "Remove LOD {0}"), LODIndex );
+			FScopedTransaction Transaction( TEXT(""), RemoveLODText, SkelMesh );
 
-	MeshDetailLayout->ForceRefreshDetails();
+			SkelMesh->Modify();
+			FSkeletalMeshUpdateContext UpdateContext;
+			UpdateContext.SkeletalMesh = SkelMesh;
+			UpdateContext.AssociatedComponents.Push(GetPersonaToolkit()->GetPreviewMeshComponent());
+
+			FLODUtilities::RemoveLOD(UpdateContext, LODIndex);
+			SkelMesh->PostEditChange();
+
+			MeshDetailLayout->ForceRefreshDetails();
+		}
+	}
 	return FReply::Handled();
 }
 
@@ -1197,6 +1599,11 @@ void FPersonaMeshDetails::CustomizeDetails( IDetailLayoutBuilder& DetailLayout )
 	const TArray<TWeakObjectPtr<UObject>>& SelectedObjects = DetailLayout.GetDetailsView().GetSelectedObjects();
 	check(SelectedObjects.Num()<=1); // The OnGenerateCustomWidgets delegate will not be useful if we try to process more than one object.
 
+	TSharedRef<IPersonaPreviewScene> PreviewScene = GetPersonaToolkit()->GetPreviewScene();
+
+	// Ensure that we only have one callback for this object registered
+	PreviewScene->RegisterOnPreviewMeshChanged(FOnPreviewMeshChanged::CreateSP(this, &FPersonaMeshDetails::OnPreviewMeshChanged));
+
 	SkeletalMeshPtr = SelectedObjects.Num() > 0 ? Cast<USkeletalMesh>(SelectedObjects[0].Get()) : nullptr;
 
 	// copy temporarily to refresh Mesh details tab from the LOD settings window
@@ -1309,6 +1716,12 @@ void FPersonaMeshDetails::OnSetPostProcessBlueprint(const FAssetData& AssetData,
 	{
 		BlueprintProperty->SetValue(SelectedBlueprint->GetAnimBlueprintGeneratedClass());
 	}
+	else if(!AssetData.IsValid())
+	{
+		// Asset data is not valid so clear the result
+		UObject* Value = nullptr;
+		BlueprintProperty->SetValue(Value);
+	}
 }
 
 FReply FPersonaMeshDetails::OnReimportLodClicked(IDetailLayoutBuilder* DetailLayout, EReimportButtonType InReimportType, int32 InLODIndex)
@@ -3232,5 +3645,13 @@ FReply FPersonaMeshDetails::OnDeleteButtonClicked(int32 LODIndex, int32 SectionI
 	return FReply::Handled();
 }
 
+void FPersonaMeshDetails::OnPreviewMeshChanged(USkeletalMesh* OldSkeletalMesh, USkeletalMesh* NewMesh)
+{
+	if (IsApplyNeeded())
+	{
+		MeshDetailLayout->ForceRefreshDetails();
+	}
+}
+
 #undef LOCTEXT_NAMESPACE
 
diff --git a/Engine/Source/Editor/Persona/Private/PersonaMeshDetails.h b/Engine/Source/Editor/Persona/Private/PersonaMeshDetails.h
--- a/Engine/Source/Editor/Persona/Private/PersonaMeshDetails.h
+++ b/Engine/Source/Editor/Persona/Private/PersonaMeshDetails.h
@@ -139,6 +139,7 @@ class FPersonaMeshDetails : public IDetailCustomization
 {
 public:
 	FPersonaMeshDetails(TSharedRef<class IPersonaToolkit> InPersonaToolkit) : PersonaToolkitPtr(InPersonaToolkit) {}
+	~FPersonaMeshDetails();
 
 	/** Makes a new instance of this detail layout class for a specific detail view requesting it */
 	static TSharedRef<IDetailCustomization> MakeInstance(TSharedRef<class IPersonaToolkit> InPersonaToolkit);
@@ -368,6 +369,7 @@ private:
 
 	/** Access the persona toolkit ptr. It should always be valid in the lifetime of this customization */
 	TSharedRef<IPersonaToolkit> GetPersonaToolkit() const { check(PersonaToolkitPtr.IsValid()); return PersonaToolkitPtr.Pin().ToSharedRef(); }
+	bool HasValidPersonaToolkit() const { return PersonaToolkitPtr.IsValid(); }
 
 	EVisibility GetOverrideUVDensityVisibililty() const;
 	ECheckBoxState IsUVDensityOverridden(int32 MaterialIndex) const;
@@ -389,6 +391,24 @@ private:
 	// Handler for reimport buttons in LOD details
 	FReply OnReimportLodClicked(IDetailLayoutBuilder* DetailLayout, EReimportButtonType InReimportType, int32 InLODIndex);
 
+	void OnCopySectionList(int32 LODIndex);
+	bool OnCanCopySectionList(int32 LODIndex) const;
+	void OnPasteSectionList(int32 LODIndex);
+
+	void OnCopySectionItem(int32 LODIndex, int32 SectionIndex);
+	bool OnCanCopySectionItem(int32 LODIndex, int32 SectionIndex) const;
+	void OnPasteSectionItem(int32 LODIndex, int32 SectionIndex);
+
+	void OnCopyMaterialList();
+	bool OnCanCopyMaterialList() const;
+	void OnPasteMaterialList();
+
+	void OnCopyMaterialItem(int32 CurrentSlot);
+	bool OnCanCopyMaterialItem(int32 CurrentSlot) const;
+	void OnPasteMaterialItem(int32 CurrentSlot);
+
+	void OnPreviewMeshChanged(USkeletalMesh* OldSkeletalMesh, USkeletalMesh* NewMesh);
+
 public:
 
 	bool IsApplyNeeded() const;
diff --git a/Engine/Source/Editor/Persona/Private/PersonaToolkit.cpp b/Engine/Source/Editor/Persona/Private/PersonaToolkit.cpp
--- a/Engine/Source/Editor/Persona/Private/PersonaToolkit.cpp
+++ b/Engine/Source/Editor/Persona/Private/PersonaToolkit.cpp
@@ -90,11 +90,13 @@ void FPersonaToolkit::CreatePreviewScene()
 		PersonaModule.OnPreviewSceneCreated().Broadcast(PreviewScene.ToSharedRef());
 
 		bool bSetMesh = false;
+		USkeletalMesh* SkeletonPreviewMesh = nullptr;
 
 		// Set the mesh
 		if (Mesh != nullptr)
 		{
 			PreviewScene->SetPreviewMesh(Mesh);
+			SkeletonPreviewMesh = Mesh;
 			bSetMesh = true;
 		}
 		else if (AnimationAsset != nullptr)
@@ -103,6 +105,7 @@ void FPersonaToolkit::CreatePreviewScene()
 			if (AssetMesh)
 			{
 				PreviewScene->SetPreviewMesh(AssetMesh);
+				SkeletonPreviewMesh = AssetMesh;
 				bSetMesh = true;
 			}
 		}
@@ -114,8 +117,11 @@ void FPersonaToolkit::CreatePreviewScene()
 			if (PreviewMesh)
 			{
 				PreviewScene->SetPreviewMesh(PreviewMesh);
+				SkeletonPreviewMesh = PreviewMesh;
 			}
 		}
+
+		EditableSkeleton->SetPreviewMesh(SkeletonPreviewMesh);
 	}
 }
 
diff --git a/Engine/Source/Editor/Persona/Private/SAnimCurvePanel.cpp b/Engine/Source/Editor/Persona/Private/SAnimCurvePanel.cpp
--- a/Engine/Source/Editor/Persona/Private/SAnimCurvePanel.cpp
+++ b/Engine/Source/Editor/Persona/Private/SAnimCurvePanel.cpp
@@ -93,6 +93,18 @@ public:
 		}
 	}
 
+	/** Returns the owner(s) of the curve */
+	virtual TArray<const UObject*> GetOwners() const override
+	{
+		TArray<const UObject*> Owners;
+		if (AnimSequenceBase.IsValid())
+		{
+			Owners.Add(AnimSequenceBase.Get());
+		}
+
+		return Owners;
+	}
+
 	/** Called to make curve owner transactional */
 	virtual void MakeTransactional() override
 	{
diff --git a/Engine/Source/Editor/Persona/Private/SAnimTrackCurvePanel.cpp b/Engine/Source/Editor/Persona/Private/SAnimTrackCurvePanel.cpp
--- a/Engine/Source/Editor/Persona/Private/SAnimTrackCurvePanel.cpp
+++ b/Engine/Source/Editor/Persona/Private/SAnimTrackCurvePanel.cpp
@@ -122,6 +122,18 @@ public:
 		}
 	}
 
+	/** Returns the owner(s) of the curve */
+	virtual TArray<const UObject*> GetOwners() const override
+	{
+		TArray<const UObject*> Owners;
+		if (BaseSequence.IsValid())
+		{
+			Owners.Add(BaseSequence.Get());
+		}
+
+		return Owners;
+	}
+
 	/** Called to make curve owner transactional */
 	virtual void MakeTransactional() override
 	{
diff --git a/Engine/Source/Editor/PlacementMode/Private/PlacementModeModule.cpp b/Engine/Source/Editor/PlacementMode/Private/PlacementModeModule.cpp
--- a/Engine/Source/Editor/PlacementMode/Private/PlacementModeModule.cpp
+++ b/Engine/Source/Editor/PlacementMode/Private/PlacementModeModule.cpp
@@ -161,6 +161,9 @@ public:
 			Category->Items.Add(CreateID(), MakeShareable( new FPlaceableItem(*UActorFactoryBasicShape::StaticClass(), FAssetData(LoadObject<UStaticMesh>(nullptr,*UActorFactoryBasicShape::BasicCylinder.ToString())), FName("ClassThumbnail.Cylinder"), BasicShapeColorOverride, SortOrder+=10, NSLOCTEXT("PlacementMode", "Cylinder", "Cylinder") )) );
 			// Cone
 			Category->Items.Add(CreateID(), MakeShareable( new FPlaceableItem(*UActorFactoryBasicShape::StaticClass(), FAssetData(LoadObject<UStaticMesh>(nullptr,*UActorFactoryBasicShape::BasicCone.ToString())), FName("ClassThumbnail.Cone"), BasicShapeColorOverride, SortOrder+=10, NSLOCTEXT("PlacementMode", "Cone", "Cone") )) );
+			// Plane
+			Category->Items.Add(CreateID(), MakeShareable(new FPlaceableItem(*UActorFactoryBasicShape::StaticClass(), FAssetData(LoadObject<UStaticMesh>(nullptr, *UActorFactoryBasicShape::BasicPlane.ToString())), FName("ClassThumbnail.Plane"), BasicShapeColorOverride, SortOrder += 10, NSLOCTEXT("PlacementMode", "Plane", "Plane"))));
+
 			Category->Items.Add(CreateID(), MakeShareable( new FPlaceableItem(*UActorFactoryTriggerBox::StaticClass(), SortOrder+=10)) );
 			Category->Items.Add(CreateID(), MakeShareable( new FPlaceableItem(*UActorFactoryTriggerSphere::StaticClass(), SortOrder+=10)) );
 		}
@@ -572,6 +575,7 @@ public:
 		Category->Items.Add(CreateID(), MakeShareable(new FPlaceableItem(*UActorFactoryBasicShape::StaticClass(), FAssetData(LoadObject<UStaticMesh>(nullptr, *UActorFactoryBasicShape::BasicSphere.ToString())), FName("ClassThumbnail.Sphere"), GetBasicShapeColorOverride())));
 		Category->Items.Add(CreateID(), MakeShareable(new FPlaceableItem(*UActorFactoryBasicShape::StaticClass(), FAssetData(LoadObject<UStaticMesh>(nullptr, *UActorFactoryBasicShape::BasicCylinder.ToString())), FName("ClassThumbnail.Cylinder"), GetBasicShapeColorOverride())));
 		Category->Items.Add(CreateID(), MakeShareable(new FPlaceableItem(*UActorFactoryBasicShape::StaticClass(), FAssetData(LoadObject<UStaticMesh>(nullptr, *UActorFactoryBasicShape::BasicCone.ToString())), FName("ClassThumbnail.Cone"), GetBasicShapeColorOverride())));
+		Category->Items.Add(CreateID(), MakeShareable(new FPlaceableItem(*UActorFactoryBasicShape::StaticClass(), FAssetData(LoadObject<UStaticMesh>(nullptr, *UActorFactoryBasicShape::BasicPlane.ToString())), FName("ClassThumbnail.Plane"), GetBasicShapeColorOverride())));
 
 		// Make a map of UClasses to ActorFactories that support them
 		const TArray< UActorFactory *>& ActorFactories = GEditor->ActorFactories;
diff --git a/Engine/Source/Editor/PropertyEditor/Private/CustomChildBuilder.cpp b/Engine/Source/Editor/PropertyEditor/Private/CustomChildBuilder.cpp
--- a/Engine/Source/Editor/PropertyEditor/Private/CustomChildBuilder.cpp
+++ b/Engine/Source/Editor/PropertyEditor/Private/CustomChildBuilder.cpp
@@ -120,6 +120,11 @@ TSharedRef<SWidget> FCustomChildrenBuilder::GenerateStructValueWidget( TSharedRe
 	}
 }
 
+IDetailCategoryBuilder& FCustomChildrenBuilder::GetParentCategory()
+{
+	return *ParentCategory.Pin();
+}
+
 FCustomChildrenBuilder& FCustomChildrenBuilder::OverrideResetChildrenToDefault(const FResetToDefaultOverride& ResetToDefault)
 {
 	CustomResetChildToDefault = ResetToDefault;
diff --git a/Engine/Source/Editor/PropertyEditor/Private/CustomChildBuilder.h b/Engine/Source/Editor/PropertyEditor/Private/CustomChildBuilder.h
--- a/Engine/Source/Editor/PropertyEditor/Private/CustomChildBuilder.h
+++ b/Engine/Source/Editor/PropertyEditor/Private/CustomChildBuilder.h
@@ -25,6 +25,7 @@ public:
 	virtual IDetailPropertyRow& AddChildProperty( TSharedRef<IPropertyHandle> PropertyHandle ) override;
 	virtual TArray<TSharedPtr<IPropertyHandle>> AddChildStructure( TSharedRef<FStructOnScope> ChildStructure ) override;
 	virtual TSharedRef<SWidget> GenerateStructValueWidget( TSharedRef<IPropertyHandle> StructPropertyHandle ) override;
+	virtual IDetailCategoryBuilder& GetParentCategory() override;
 
 	const TArray< FDetailLayoutCustomization >& GetChildCustomizations() const { return ChildCustomizations; }
 
diff --git a/Engine/Source/Editor/PropertyEditor/Private/DetailCategoryBuilderImpl.h b/Engine/Source/Editor/PropertyEditor/Private/DetailCategoryBuilderImpl.h
--- a/Engine/Source/Editor/PropertyEditor/Private/DetailCategoryBuilderImpl.h
+++ b/Engine/Source/Editor/PropertyEditor/Private/DetailCategoryBuilderImpl.h
@@ -165,7 +165,7 @@ public:
 	virtual void GetDefaultProperties( TArray<TSharedRef<IPropertyHandle> >& OutAllProperties, bool bSimpleProperties = true, bool bAdvancedProperties = true ) override;
 	virtual const FText& GetDisplayName() const override { return DisplayName; }
 	virtual void SetCategoryVisibility( bool bIsVisible ) override;
-
+	
 	/** IDetailTreeNode interface */
 	virtual IDetailsViewPrivate& GetDetailsView() const override{ return DetailLayoutBuilder.Pin()->GetDetailsView(); }
 	virtual TSharedRef< ITableRow > GenerateNodeWidget( const TSharedRef<STableViewBase>& OwnerTable, const FDetailColumnSizeData& ColumnSizeData, const TSharedRef<IPropertyUtilities>& PropertyUtilities, bool bAllowFavoriteSystem) override;
@@ -177,6 +177,11 @@ public:
 	virtual bool ShouldShowOnlyChildren() const override { return false; }
 	virtual FName GetNodeName() const override { return GetCategoryName(); }
 
+	/**
+	 * @return true if the parent layout is valid or has been destroyed by a refresh.
+	 */
+	bool IsParentLayoutValid() const { return DetailLayoutBuilder.IsValid(); }
+
 	/**
 	 * @return The name of the category
 	 */
diff --git a/Engine/Source/Editor/PropertyEditor/Private/DetailGroup.cpp b/Engine/Source/Editor/PropertyEditor/Private/DetailGroup.cpp
--- a/Engine/Source/Editor/PropertyEditor/Private/DetailGroup.cpp
+++ b/Engine/Source/Editor/PropertyEditor/Private/DetailGroup.cpp
@@ -58,6 +58,16 @@ IDetailPropertyRow& FDetailGroup::AddPropertyRow( TSharedRef<IPropertyHandle> Pr
 	return *NewCustomization.PropertyRow;
 }
 
+IDetailGroup& FDetailGroup::AddGroup(FName NewGroupName, const FText& InLocalizedDisplayName, bool bInStartExpanded)
+{
+	FDetailLayoutCustomization NewCustomization;
+	NewCustomization.DetailGroup = MakeShareable(new FDetailGroup(NewGroupName, ParentCategory.Pin().ToSharedRef(), InLocalizedDisplayName, bStartExpanded));
+
+	GroupChildren.Add(NewCustomization);
+
+	return *NewCustomization.DetailGroup;
+}
+
 void FDetailGroup::ToggleExpansion( bool bExpand )
 {
 	if( ParentCategory.IsValid() && OwnerTreeNode.IsValid() )
diff --git a/Engine/Source/Editor/PropertyEditor/Private/DetailGroup.h b/Engine/Source/Editor/PropertyEditor/Private/DetailGroup.h
--- a/Engine/Source/Editor/PropertyEditor/Private/DetailGroup.h
+++ b/Engine/Source/Editor/PropertyEditor/Private/DetailGroup.h
@@ -25,6 +25,8 @@ public:
 	virtual IDetailPropertyRow& HeaderProperty( TSharedRef<IPropertyHandle> PropertyHandle ) override;
 	virtual FDetailWidgetRow& AddWidgetRow() override;
 	virtual IDetailPropertyRow& AddPropertyRow( TSharedRef<IPropertyHandle> PropertyHandle ) override;
+	virtual IDetailGroup& AddGroup(FName NewGroupName, const FText& InLocalizedDisplayName, bool bInStartExpanded = false) override;
+
 	virtual void ToggleExpansion( bool bExpand ) override;
 	virtual bool GetExpansionState() const override;
 
diff --git a/Engine/Source/Editor/PropertyEditor/Private/DetailItemNode.cpp b/Engine/Source/Editor/PropertyEditor/Private/DetailItemNode.cpp
--- a/Engine/Source/Editor/PropertyEditor/Private/DetailItemNode.cpp
+++ b/Engine/Source/Editor/PropertyEditor/Private/DetailItemNode.cpp
@@ -331,11 +331,12 @@ static bool PassesAllFilters( const FDetailLayoutCustomization& InCustomization,
 			// The property node is visible (note categories are never visible unless they have a child that is visible )
 			bPassesAllFilters = (bPassesSearchFilter && bPassesModifiedFilter && bPassesDifferingFilter) || bPassesCategoryFilter;
 		}
-		else if( InCustomization.HasCustomWidget() )
+		else if (InCustomization.HasCustomWidget())
 		{
 			const bool bPassesTextFilter = Local::StringPassesFilter(InFilter, InCustomization.WidgetDecl->FilterTextString.ToString());
+			const bool bPassesModifiedFilter = (InFilter.bShowOnlyModifiedProperties == false || InCustomization.WidgetDecl->DiffersFromDefaultAttr.Get() == true);
 
-			bPassesAllFilters = bPassesTextFilter || bPassesCategoryFilter;
+			bPassesAllFilters = (bPassesTextFilter && bPassesModifiedFilter) || bPassesCategoryFilter;
 		}
 	}
 
diff --git a/Engine/Source/Editor/PropertyEditor/Private/DetailLayoutBuilderImpl.cpp b/Engine/Source/Editor/PropertyEditor/Private/DetailLayoutBuilderImpl.cpp
--- a/Engine/Source/Editor/PropertyEditor/Private/DetailLayoutBuilderImpl.cpp
+++ b/Engine/Source/Editor/PropertyEditor/Private/DetailLayoutBuilderImpl.cpp
@@ -183,8 +183,18 @@ void FDetailLayoutBuilderImpl::GenerateDetailLayout()
 	TArray< TSharedRef<FDetailCategoryImpl> > SimpleCategories;
 	TArray< TSharedRef<FDetailCategoryImpl> > AdvancedOnlyCategories;
 
-	BuildCategories( CustomCategoryMap, SimpleCategories, AdvancedOnlyCategories );
-	BuildCategories( DefaultCategoryMap, SimpleCategories, AdvancedOnlyCategories );
+	// Customizations can add more categories while customizing so just keep doing this until the maps are empty
+	while(DefaultCategoryMap.Num() > 0 || CustomCategoryMap.Num() > 0)
+	{
+		FCategoryMap DefaultCategoryMapCopy = DefaultCategoryMap;
+		FCategoryMap CustomCategoryMapCopy = CustomCategoryMap;
+
+		DefaultCategoryMap.Empty();
+		CustomCategoryMap.Empty();
+
+		BuildCategories(DefaultCategoryMapCopy, SimpleCategories, AdvancedOnlyCategories);
+		BuildCategories(CustomCategoryMapCopy, SimpleCategories, AdvancedOnlyCategories);
+	}
 
 	SimpleCategories.Sort( FCompareFDetailCategoryImpl() );
 	AdvancedOnlyCategories.Sort( FCompareFDetailCategoryImpl() );
@@ -194,12 +204,12 @@ void FDetailLayoutBuilderImpl::GenerateDetailLayout()
 	// Merge the two category lists in sorted order
 	for (int32 CategoryIndex = 0; CategoryIndex < SimpleCategories.Num(); ++CategoryIndex)
 	{
-		CategoryNodes.Add(SimpleCategories[CategoryIndex]);
+		CategoryNodes.AddUnique(SimpleCategories[CategoryIndex]);
 	}
 
 	for (int32 CategoryIndex = 0; CategoryIndex < AdvancedOnlyCategories.Num(); ++CategoryIndex)
 	{
-		CategoryNodes.Add(AdvancedOnlyCategories[CategoryIndex]);
+		CategoryNodes.AddUnique(AdvancedOnlyCategories[CategoryIndex]);
 	}
 
 	if(DetailsView.ContainsMultipleTopLevelObjects())
diff --git a/Engine/Source/Editor/PropertyEditor/Private/Presentation/PropertyTable/PropertyTableColumn.cpp b/Engine/Source/Editor/PropertyEditor/Private/Presentation/PropertyTable/PropertyTableColumn.cpp
--- a/Engine/Source/Editor/PropertyEditor/Private/Presentation/PropertyTable/PropertyTableColumn.cpp
+++ b/Engine/Source/Editor/PropertyEditor/Private/Presentation/PropertyTable/PropertyTableColumn.cpp
@@ -187,7 +187,7 @@ FORCEINLINE bool FCompareRowByColumnAscending<UNameProperty>::ComparePropertyVal
 }
 
 template<>
-FORCEINLINE bool FCompareRowByColumnAscending<UObjectProperty>::ComparePropertyValue( const TSharedPtr< IPropertyHandle >& LhsPropertyHandle, const TSharedPtr< IPropertyHandle >& RhsPropertyHandle ) const
+FORCEINLINE bool FCompareRowByColumnAscending<UObjectPropertyBase>::ComparePropertyValue( const TSharedPtr< IPropertyHandle >& LhsPropertyHandle, const TSharedPtr< IPropertyHandle >& RhsPropertyHandle ) const
 {
 	UObject* LhsValue; 
 	LhsPropertyHandle->GetValue( LhsValue );
@@ -205,7 +205,50 @@ FORCEINLINE bool FCompareRowByColumnAscending<UObjectProperty>::ComparePropertyV
 		return false;
 	}
 
-	return LhsValue->GetFullName() < RhsValue->GetFullName();
+	return LhsValue->GetName() < RhsValue->GetName();
+}
+
+template<>
+FORCEINLINE bool FCompareRowByColumnAscending<UStructProperty>::ComparePropertyValue( const TSharedPtr< IPropertyHandle >& LhsPropertyHandle, const TSharedPtr< IPropertyHandle >& RhsPropertyHandle ) const
+{
+	if ( !FPropertyTableColumn::IsSupportedStructProperty(LhsPropertyHandle->GetProperty() ) )
+	{
+		return true;
+	}
+
+	if ( !FPropertyTableColumn::IsSupportedStructProperty(RhsPropertyHandle->GetProperty() ) )
+	{
+		return false;
+	}
+
+	{
+		FVector LhsVector;
+		FVector RhsVector;
+
+		if ( LhsPropertyHandle->GetValue(LhsVector) != FPropertyAccess::Fail && RhsPropertyHandle->GetValue(RhsVector) != FPropertyAccess::Fail )
+		{
+			return LhsVector.SizeSquared() < RhsVector.SizeSquared();
+		}
+
+		FVector2D LhsVector2D;
+		FVector2D RhsVector2D;
+
+		if ( LhsPropertyHandle->GetValue(LhsVector2D) != FPropertyAccess::Fail && RhsPropertyHandle->GetValue(RhsVector2D) != FPropertyAccess::Fail )
+		{
+			return LhsVector2D.SizeSquared()  < RhsVector2D.SizeSquared();
+		}
+
+		FVector4 LhsVector4;
+		FVector4 RhsVector4;
+
+		if ( LhsPropertyHandle->GetValue(LhsVector4) != FPropertyAccess::Fail && RhsPropertyHandle->GetValue(RhsVector4) != FPropertyAccess::Fail )
+		{
+			return LhsVector4.SizeSquared()  < RhsVector4.SizeSquared();
+		}
+	}
+
+	ensureMsgf(false, TEXT("A supported struct property does not have a defined implementation for sorting a property column."));
+	return false;
 }
 
 
@@ -426,7 +469,8 @@ bool FPropertyTableColumn::CanSortBy() const
 			Property->IsA( UFloatProperty::StaticClass() ) ||
 			Property->IsA( UNameProperty::StaticClass() )  ||
 			Property->IsA( UStrProperty::StaticClass() )   ||
-			( Property->IsA( UObjectProperty::StaticClass() ) && !Property->HasAnyPropertyFlags(CPF_InstancedReference) );
+			IsSupportedStructProperty( Property )		   ||
+			( Property->IsA( UObjectPropertyBase::StaticClass() ) && !Property->HasAnyPropertyFlags(CPF_InstancedReference) );
 			//Property->IsA( UTextProperty::StaticClass() );
 	}
 
@@ -545,17 +589,30 @@ void FPropertyTableColumn::Sort( TArray< TSharedRef< class IPropertyTableRow > >
 			Rows.Sort( FCompareRowByColumnDescending< UStrProperty >( SharedThis( this ), StrProperty ) );
 		}
 	}
-	else if ( Property->IsA( UObjectProperty::StaticClass() ) )
+	else if ( Property->IsA( UObjectPropertyBase::StaticClass() ) )
 	{
-		TWeakObjectPtr<UObjectProperty> ObjectProperty( Cast< UObjectProperty >( Property ) );
+		TWeakObjectPtr<UObjectPropertyBase> ObjectProperty( Cast< UObjectPropertyBase >( Property ) );
 
 		if ( SortMode == EColumnSortMode::Ascending )
 		{
-			Rows.Sort( FCompareRowByColumnAscending< UObjectProperty >( SharedThis( this ), ObjectProperty ) );
+			Rows.Sort( FCompareRowByColumnAscending< UObjectPropertyBase >( SharedThis( this ), ObjectProperty ) );
 		}
 		else
 		{
-			Rows.Sort( FCompareRowByColumnDescending< UObjectProperty >( SharedThis( this ), ObjectProperty ) );
+			Rows.Sort( FCompareRowByColumnDescending< UObjectPropertyBase >( SharedThis( this ), ObjectProperty ) );
+		}
+	}
+	else if ( IsSupportedStructProperty( Property ) )
+	{
+		TWeakObjectPtr<UStructProperty> StructProperty(Cast < UStructProperty >( Property ) );
+
+		if (SortMode == EColumnSortMode::Ascending)
+		{
+			Rows.Sort( FCompareRowByColumnAscending< UStructProperty >( SharedThis(this), StructProperty ) );
+		}
+		else
+		{
+			Rows.Sort( FCompareRowByColumnDescending< UStructProperty >( SharedThis(this), StructProperty ) );
 		}
 	}
 	//else if ( Property->IsA( UTextProperty::StaticClass() ) )
@@ -609,4 +666,19 @@ void FPropertyTableColumn::SetFrozen(bool InIsFrozen)
 	FrozenStateChanged.Broadcast( SharedThis(this) );
 }
 
+bool FPropertyTableColumn::IsSupportedStructProperty(const UProperty* InProperty)
+{
+	if ( InProperty != nullptr && Cast<UStructProperty>(InProperty) != nullptr)
+	{
+		const UStructProperty* StructProp = Cast<UStructProperty>(InProperty);
+		FName StructName = StructProp->Struct->GetFName();
+
+		return StructName == NAME_Vector ||
+			StructName == NAME_Vector2D	 ||
+			StructName == NAME_Vector4;
+	}
+
+	return false;
+}
+
 #undef LOCTEXT_NAMESPACE
diff --git a/Engine/Source/Editor/PropertyEditor/Private/Presentation/PropertyTable/PropertyTableColumn.h b/Engine/Source/Editor/PropertyEditor/Private/Presentation/PropertyTable/PropertyTableColumn.h
--- a/Engine/Source/Editor/PropertyEditor/Private/Presentation/PropertyTable/PropertyTableColumn.h
+++ b/Engine/Source/Editor/PropertyEditor/Private/Presentation/PropertyTable/PropertyTableColumn.h
@@ -64,6 +64,8 @@ public:
 
 	//~ End IPropertyTable Interface
 
+	static bool IsSupportedStructProperty(const UProperty* InProperty);
+
 private:
 
 	void GenerateColumnId();
diff --git a/Engine/Source/Editor/PropertyEditor/Private/PropertyCustomizationHelpers.cpp b/Engine/Source/Editor/PropertyEditor/Private/PropertyCustomizationHelpers.cpp
--- a/Engine/Source/Editor/PropertyEditor/Private/PropertyCustomizationHelpers.cpp
+++ b/Engine/Source/Editor/PropertyEditor/Private/PropertyCustomizationHelpers.cpp
@@ -477,6 +477,8 @@ void SClassPropertyEntryBox::Construct(const FArguments& InArgs)
 				.AllowAbstract(InArgs._AllowAbstract)
 				.IsBlueprintBaseOnly(InArgs._IsBlueprintBaseOnly)
 				.AllowNone(InArgs._AllowNone)
+				.ShowViewOptions(!InArgs._HideViewOptions)
+				.ShowTree(InArgs._ShowTreeView)
 				.SelectedClass(InArgs._SelectedClass)
 				.OnSetClass(InArgs._OnSetClass)
 		]
@@ -1003,6 +1005,9 @@ void FMaterialList::Tick( float DeltaTime )
 
 void FMaterialList::GenerateHeaderRowContent( FDetailWidgetRow& NodeRow )
 {
+	NodeRow.CopyAction(FUIAction(FExecuteAction::CreateSP(this, &FMaterialList::OnCopyMaterialList), FCanExecuteAction::CreateSP(this, &FMaterialList::OnCanCopyMaterialList)));
+	NodeRow.PasteAction(FUIAction(FExecuteAction::CreateSP(this, &FMaterialList::OnPasteMaterialList)));
+
 	if (bAllowCollpase)
 	{
 		NodeRow.NameContent()
@@ -1104,6 +1109,58 @@ void FMaterialList::GenerateChildContent( IDetailChildrenBuilder& ChildrenBuilde
 		];
 	}		
 }
+
+bool FMaterialList::OnCanCopyMaterialList() const
+{
+	if (MaterialListDelegates.OnCanCopyMaterialList.IsBound())
+	{
+		return MaterialListDelegates.OnCanCopyMaterialList.Execute();
+	}
+
+	return false;
+}
+
+void FMaterialList::OnCopyMaterialList()
+{
+	if (MaterialListDelegates.OnCopyMaterialList.IsBound())
+	{
+		MaterialListDelegates.OnCopyMaterialList.Execute();
+	}
+}
+
+void FMaterialList::OnPasteMaterialList()
+{
+	if (MaterialListDelegates.OnPasteMaterialList.IsBound())
+	{
+		MaterialListDelegates.OnPasteMaterialList.Execute();
+	}
+}
+
+bool FMaterialList::OnCanCopyMaterialItem(int32 CurrentSlot) const
+{
+	if (MaterialListDelegates.OnCanCopyMaterialItem.IsBound())
+	{
+		return MaterialListDelegates.OnCanCopyMaterialItem.Execute(CurrentSlot);
+	}
+
+	return false;
+}
+
+void FMaterialList::OnCopyMaterialItem(int32 CurrentSlot)
+{
+	if (MaterialListDelegates.OnCopyMaterialItem.IsBound())
+	{
+		MaterialListDelegates.OnCopyMaterialItem.Execute(CurrentSlot);
+	}
+}
+
+void FMaterialList::OnPasteMaterialItem(int32 CurrentSlot)
+{
+	if (MaterialListDelegates.OnPasteMaterialItem.IsBound())
+	{
+		MaterialListDelegates.OnPasteMaterialItem.Execute(CurrentSlot);
+	}
+}
 			
 void FMaterialList::AddMaterialItem( FDetailWidgetRow& Row, int32 CurrentSlot, const FMaterialListItem& Item, bool bDisplayLink )
 {
@@ -1135,6 +1192,9 @@ void FMaterialList::AddMaterialItem( FDetailWidgetRow& Row, int32 CurrentSlot, c
 		ViewedMaterials.Add( NewView );
 	}
 
+	Row.CopyAction(FUIAction(FExecuteAction::CreateSP(this, &FMaterialList::OnCopyMaterialItem, Item.SlotIndex), FCanExecuteAction::CreateSP(this, &FMaterialList::OnCanCopyMaterialItem, Item.SlotIndex)));
+	Row.PasteAction(FUIAction(FExecuteAction::CreateSP(this, &FMaterialList::OnPasteMaterialItem, Item.SlotIndex)));
+
 	Row.NameContent()
 	[
 		NewView->CreateNameContent()
@@ -1696,6 +1756,9 @@ void FSectionList::Tick(float DeltaTime)
 
 void FSectionList::GenerateHeaderRowContent(FDetailWidgetRow& NodeRow)
 {
+	NodeRow.CopyAction(FUIAction(FExecuteAction::CreateSP(this, &FSectionList::OnCopySectionList), FCanExecuteAction::CreateSP(this, &FSectionList::OnCanCopySectionList)));
+	NodeRow.PasteAction(FUIAction(FExecuteAction::CreateSP(this, &FSectionList::OnPasteSectionList)));
+
 	NodeRow.NameContent()
 	[
 		SNew(STextBlock)
@@ -1748,6 +1811,58 @@ void FSectionList::GenerateChildContent(IDetailChildrenBuilder& ChildrenBuilder)
 	}
 }
 
+bool FSectionList::OnCanCopySectionList() const
+{
+	if (SectionListDelegates.OnCanCopySectionList.IsBound())
+	{
+		return SectionListDelegates.OnCanCopySectionList.Execute();
+	}
+
+	return false;
+}
+
+void FSectionList::OnCopySectionList()
+{
+	if (SectionListDelegates.OnCopySectionList.IsBound())
+	{
+		SectionListDelegates.OnCopySectionList.Execute();
+	}
+}
+
+void FSectionList::OnPasteSectionList()
+{
+	if (SectionListDelegates.OnPasteSectionList.IsBound())
+	{
+		SectionListDelegates.OnPasteSectionList.Execute();
+	}
+}
+
+bool FSectionList::OnCanCopySectionItem(int32 LODIndex, int32 SectionIndex) const
+{
+	if (SectionListDelegates.OnCanCopySectionItem.IsBound())
+	{
+		return SectionListDelegates.OnCanCopySectionItem.Execute(LODIndex, SectionIndex);
+	}
+
+	return false;
+}
+
+void FSectionList::OnCopySectionItem(int32 LODIndex, int32 SectionIndex)
+{
+	if (SectionListDelegates.OnCopySectionItem.IsBound())
+	{
+		SectionListDelegates.OnCopySectionItem.Execute(LODIndex, SectionIndex);
+	}
+}
+
+void FSectionList::OnPasteSectionItem(int32 LODIndex, int32 SectionIndex)
+{
+	if (SectionListDelegates.OnPasteSectionItem.IsBound())
+	{
+		SectionListDelegates.OnPasteSectionItem.Execute(LODIndex, SectionIndex);
+	}
+}
+
 void FSectionList::AddSectionItem(FDetailWidgetRow& Row, int32 LodIndex, const struct FSectionListItem& Item, bool bDisplayLink)
 {
 	uint32 NumSections = SectionListBuilder->GetNumSections(LodIndex);
@@ -1778,6 +1893,9 @@ void FSectionList::AddSectionItem(FDetailWidgetRow& Row, int32 LodIndex, const s
 		ViewedSections.Add(NewView);
 	}
 
+	Row.CopyAction(FUIAction(FExecuteAction::CreateSP(this, &FSectionList::OnCopySectionItem, LodIndex, Item.SectionIndex), FCanExecuteAction::CreateSP(this, &FSectionList::OnCanCopySectionItem, LodIndex, Item.SectionIndex)));
+	Row.PasteAction(FUIAction(FExecuteAction::CreateSP(this, &FSectionList::OnPasteSectionItem, LodIndex, Item.SectionIndex)));
+
 	Row.NameContent()
 		[
 			NewView->CreateNameContent()
diff --git a/Engine/Source/Editor/PropertyEditor/Private/PropertyHandleImpl.cpp b/Engine/Source/Editor/PropertyEditor/Private/PropertyHandleImpl.cpp
--- a/Engine/Source/Editor/PropertyEditor/Private/PropertyHandleImpl.cpp
+++ b/Engine/Source/Editor/PropertyEditor/Private/PropertyHandleImpl.cpp
@@ -453,7 +453,9 @@ FPropertyAccess::Result FPropertyValueImpl::ImportText( const TArray<FObjectBase
 					
 					if (bIsInContainer)
 					{
-						if ( SetProp->HasElement( ParentNode->GetValueBaseAddress((uint8*)Cur.Object), Cur.BaseAddress, InValues[ObjectIndex])  )
+						FScriptSetHelper SetHelper(SetProp, ParentNode->GetValueBaseAddress((uint8*)Cur.Object));
+
+						if ( SetHelper.HasElement(Cur.BaseAddress, InValues[ObjectIndex]) )
 						{
 							// Duplicate element in the set
 							ShowInvalidOperationError(LOCTEXT("DuplicateSetElement", "Duplicate elements are not allowed in Set properties."));
@@ -469,7 +471,9 @@ FPropertyAccess::Result FPropertyValueImpl::ImportText( const TArray<FObjectBase
 
 					if (bIsInContainer)
 					{
-						if ( MapProperty->HasKey(ParentNode->GetValueBaseAddress((uint8*)Cur.Object), Cur.BaseAddress, InValues[ObjectIndex]) )
+						FScriptMapHelper MapHelper(MapProperty, ParentNode->GetValueBaseAddress((uint8*)Cur.Object));
+
+						if ( MapHelper.HasKey(Cur.BaseAddress, InValues[ObjectIndex]) )
 						{
 							// Duplicate key in the map
 							ShowInvalidOperationError(LOCTEXT("DuplicateMapKey", "Duplicate keys are not allowed in Map properties."));
diff --git a/Engine/Source/Editor/PropertyEditor/Private/PropertyHandleImpl.h b/Engine/Source/Editor/PropertyEditor/Private/PropertyHandleImpl.h
--- a/Engine/Source/Editor/PropertyEditor/Private/PropertyHandleImpl.h
+++ b/Engine/Source/Editor/PropertyEditor/Private/PropertyHandleImpl.h
@@ -665,7 +665,7 @@ public:
 	virtual FPropertyAccess::Result GetNumElements(uint32& OutNumElements) override;
 	virtual void SetOnNumElementsChanged(FSimpleDelegate& InOnNumElementsChanged) override;
 	virtual bool HasDocumentation() override { return true; }
-	virtual FString GetDocumentationLink() override { return FString("Programming/UnrealArchitecture/Reference/Properties/"); }	// @todo: needs a better documentation page
+	virtual FString GetDocumentationLink() override { return FString("Engine/UI/LevelEditor/Details/Properties/Set/"); }	// @todo: needs a better documentation page
 	virtual FString GetDocumentationExcerptName() override { return FString("Sets"); }
 private:
 	/**
@@ -689,7 +689,7 @@ public:
 	virtual FPropertyAccess::Result GetNumElements(uint32& OutNumElements) override;
 	virtual void SetOnNumElementsChanged(FSimpleDelegate& InOnNumElementsChanged) override;
 	virtual bool HasDocumentation() override { return true; }
-	virtual FString GetDocumentationLink() override { return FString("Programming/UnrealArchitecture/TMap"); }
+	virtual FString GetDocumentationLink() override { return FString("Engine/UI/LevelEditor/Details/Properties/Map/"); }
 	virtual FString GetDocumentationExcerptName() override { return FString("Maps"); }
 private:
 	/**
diff --git a/Engine/Source/Editor/PropertyEditor/Private/SDetailSingleItemRow.cpp b/Engine/Source/Editor/PropertyEditor/Private/SDetailSingleItemRow.cpp
--- a/Engine/Source/Editor/PropertyEditor/Private/SDetailSingleItemRow.cpp
+++ b/Engine/Source/Editor/PropertyEditor/Private/SDetailSingleItemRow.cpp
@@ -190,147 +190,156 @@ void SDetailSingleItemRow::Construct( const FArguments& InArgs, FDetailLayoutCus
 	EHorizontalAlignment HorizontalAlignment = HAlign_Fill;
 	EVerticalAlignment VerticalAlignment = VAlign_Fill;
 
-	if( InCustomization->IsValidCustomization() )
+	const bool bIsValidTreeNode = InOwnerTreeNode->GetParentCategory().IsValid() && InOwnerTreeNode->GetParentCategory()->IsParentLayoutValid();
+	if(bIsValidTreeNode)
 	{
-		FDetailWidgetRow Row = InCustomization->GetWidgetRow();
-
-		TSharedPtr<SWidget> NameWidget;
-		TSharedPtr<SWidget> ValueWidget;
-	
-		NameWidget = Row.NameWidget.Widget;
-		if( Row.IsEnabledAttr.IsBound() )
+		if(InCustomization->IsValidCustomization())
 		{
-			NameWidget->SetEnabled( Row.IsEnabledAttr );
-		}
-		
-		ValueWidget =
-			SNew( SConstrainedBox )
-			.MinWidth( Row.ValueWidget.MinWidth )
-			.MaxWidth( Row.ValueWidget.MaxWidth )
-			[
-				Row.ValueWidget.Widget
-			];
+			FDetailWidgetRow Row = InCustomization->GetWidgetRow();
 
-		ValueWidget = CreateExtensionWidget(ValueWidget.ToSharedRef(), *Customization, InOwnerTreeNode );
+			TSharedPtr<SWidget> NameWidget;
+			TSharedPtr<SWidget> ValueWidget;
 
-		if( Row.IsEnabledAttr.IsBound() )
-		{
-			ValueWidget->SetEnabled( Row.IsEnabledAttr );
-		}
+			NameWidget = Row.NameWidget.Widget;
+			if(Row.IsEnabledAttr.IsBound())
+			{
+				NameWidget->SetEnabled(Row.IsEnabledAttr);
+			}
+
+			ValueWidget =
+				SNew(SConstrainedBox)
+				.MinWidth(Row.ValueWidget.MinWidth)
+				.MaxWidth(Row.ValueWidget.MaxWidth)
+				[
+					Row.ValueWidget.Widget
+				];
 
-		TSharedRef<SWidget> KeyFrameButton = CreateKeyframeButton( *Customization, InOwnerTreeNode );
-		TAttribute<bool> IsPropertyEditingEnabled = InOwnerTreeNode->IsPropertyEditingEnabled();
+			ValueWidget = CreateExtensionWidget(ValueWidget.ToSharedRef(), *Customization, InOwnerTreeNode);
 
-		bool const bEnableFavoriteSystem = GIsRequestingExit ? false : (GetDefault<UEditorExperimentalSettings>()->bEnableFavoriteSystem && bAllowFavoriteSystem);
+			if(Row.IsEnabledAttr.IsBound())
+			{
+				ValueWidget->SetEnabled(Row.IsEnabledAttr);
+			}
 
-		TSharedRef<SHorizontalBox> InternalLeftColumnRowBox = SNew(SHorizontalBox);
+			TSharedRef<SWidget> KeyFrameButton = CreateKeyframeButton(*Customization, InOwnerTreeNode);
+			TAttribute<bool> IsPropertyEditingEnabled = InOwnerTreeNode->IsPropertyEditingEnabled();
 
-		if (bEnableFavoriteSystem)
-		{
-			InternalLeftColumnRowBox->AddSlot()
-				.Padding(0.0f, 0.0f)
-				.HAlign(HAlign_Left)
-				.VAlign(VAlign_Center)
-				.AutoWidth()
-				[
-					SNew(SButton)
-					.HAlign(HAlign_Center)
+			bool const bEnableFavoriteSystem = GIsRequestingExit ? false : (GetDefault<UEditorExperimentalSettings>()->bEnableFavoriteSystem && bAllowFavoriteSystem);
+
+			TSharedRef<SHorizontalBox> InternalLeftColumnRowBox = SNew(SHorizontalBox);
+
+			if(bEnableFavoriteSystem)
+			{
+				InternalLeftColumnRowBox->AddSlot()
+					.Padding(0.0f, 0.0f)
+					.HAlign(HAlign_Left)
 					.VAlign(VAlign_Center)
-					.IsFocusable(false)
-					.ButtonStyle(FEditorStyle::Get(), "NoBorder")
-					.OnClicked(this, &SDetailSingleItemRow::OnFavoriteToggle)
+					.AutoWidth()
 					[
-						SNew(SImage).Image(this, &SDetailSingleItemRow::GetFavoriteButtonBrush)
-					]
-				];
-		}
-		InternalLeftColumnRowBox->AddSlot()
-			.Padding(3.0f, 0.0f)
-			.HAlign(HAlign_Left)
-			.VAlign(VAlign_Center)
-			.AutoWidth()
-			[
-				SNew(SExpanderArrow, SharedThis(this))
-				.BaseIndentLevel(1)
-			];
-		
-
-		if( bHasMultipleColumns )
-		{
-			NameWidget->SetEnabled(IsPropertyEditingEnabled);
-			TSharedPtr<SHorizontalBox> HBox;
-			
-			InternalLeftColumnRowBox->AddSlot()
-				.HAlign(Row.NameWidget.HorizontalAlignment)
-				.VAlign(Row.NameWidget.VerticalAlignment)
-				.Padding(DetailWidgetConstants::LeftRowPadding)
-				[
-					NameWidget.ToSharedRef()
-				];
+						SNew(SButton)
+						.HAlign(HAlign_Center)
+						.VAlign(VAlign_Center)
+						.IsFocusable(false)
+						.ButtonStyle(FEditorStyle::Get(), "NoBorder")
+						.OnClicked(this, &SDetailSingleItemRow::OnFavoriteToggle)
+						[
+							SNew(SImage).Image(this, &SDetailSingleItemRow::GetFavoriteButtonBrush)
+						]
+					];
+			}
 			InternalLeftColumnRowBox->AddSlot()
 				.Padding(3.0f, 0.0f)
-				.HAlign(HAlign_Right)
+				.HAlign(HAlign_Left)
 				.VAlign(VAlign_Center)
 				.AutoWidth()
 				[
-					KeyFrameButton
+					SNew(SExpanderArrow, SharedThis(this))
+					.BaseIndentLevel(1)
 				];
 
-			TSharedRef<SSplitter> Splitter = 
-				SNew( SSplitter )
-				.Style( FEditorStyle::Get(), "DetailsView.Splitter" )
-				.PhysicalSplitterHandleSize( 1.0f )
-				.HitDetectionSplitterHandleSize( 5.0f )
-				+ SSplitter::Slot()
-				.Value( ColumnSizeData.LeftColumnWidth )
-				.OnSlotResized( SSplitter::FOnSlotResized::CreateSP( this, &SDetailSingleItemRow::OnLeftColumnResized ) )
-				[
-					InternalLeftColumnRowBox
-				]
-				+ SSplitter::Slot()
-				.Value( ColumnSizeData.RightColumnWidth )
-				.OnSlotResized( ColumnSizeData.OnWidthChanged )
-				[
-					SNew( SHorizontalBox )
-					+ SHorizontalBox::Slot()
+
+			if(bHasMultipleColumns)
+			{
+				NameWidget->SetEnabled(IsPropertyEditingEnabled);
+				TSharedPtr<SHorizontalBox> HBox;
+
+				InternalLeftColumnRowBox->AddSlot()
+					.HAlign(Row.NameWidget.HorizontalAlignment)
+					.VAlign(Row.NameWidget.VerticalAlignment)
+					.Padding(DetailWidgetConstants::LeftRowPadding)
+					[
+						NameWidget.ToSharedRef()
+					];
+				InternalLeftColumnRowBox->AddSlot()
+					.Padding(3.0f, 0.0f)
+					.HAlign(HAlign_Right)
+					.VAlign(VAlign_Center)
+					.AutoWidth()
+					[
+						KeyFrameButton
+					];
+
+				TSharedRef<SSplitter> Splitter =
+					SNew(SSplitter)
+					.Style(FEditorStyle::Get(), "DetailsView.Splitter")
+					.PhysicalSplitterHandleSize(1.0f)
+					.HitDetectionSplitterHandleSize(5.0f)
+					+ SSplitter::Slot()
+					.Value(ColumnSizeData.LeftColumnWidth)
+					.OnSlotResized(SSplitter::FOnSlotResized::CreateSP(this, &SDetailSingleItemRow::OnLeftColumnResized))
 					[
-						SAssignNew(HBox, SHorizontalBox)
+						InternalLeftColumnRowBox
+					]
+					+ SSplitter::Slot()
+					.Value(ColumnSizeData.RightColumnWidth)
+					.OnSlotResized(ColumnSizeData.OnWidthChanged)
+					[
+						SNew(SHorizontalBox)
 						+ SHorizontalBox::Slot()
-						.Padding(DetailWidgetConstants::RightRowPadding)
-						.HAlign(Row.ValueWidget.HorizontalAlignment)
-						.VAlign(Row.ValueWidget.VerticalAlignment)
 						[
-							SNew(SBox)
-							.IsEnabled(IsPropertyEditingEnabled)
+							SAssignNew(HBox, SHorizontalBox)
+							+ SHorizontalBox::Slot()
+							.Padding(DetailWidgetConstants::RightRowPadding)
+							.HAlign(Row.ValueWidget.HorizontalAlignment)
+							.VAlign(Row.ValueWidget.VerticalAlignment)
 							[
-								ValueWidget.ToSharedRef()
+								SNew(SBox)
+								.IsEnabled(IsPropertyEditingEnabled)
+								[
+									ValueWidget.ToSharedRef()
+								]
 							]
 						]
-					]
-				];
+					];
 				Widget = Splitter;
-		}
-		else
-		{
-			Row.WholeRowWidget.Widget->SetEnabled(IsPropertyEditingEnabled);
-			InternalLeftColumnRowBox->AddSlot()
-				.HAlign(Row.WholeRowWidget.HorizontalAlignment)
-				.VAlign(Row.WholeRowWidget.VerticalAlignment)
-				.Padding(DetailWidgetConstants::LeftRowPadding)
-				[
-					Row.WholeRowWidget.Widget
-				];
-			InternalLeftColumnRowBox->AddSlot()
-				.Padding(3.0f, 0.0f)
-				.HAlign(HAlign_Right)
-				.VAlign(VAlign_Center)
-				[
-					KeyFrameButton
-				];
-			Widget = InternalLeftColumnRowBox;
+			}
+			else
+			{
+				Row.WholeRowWidget.Widget->SetEnabled(IsPropertyEditingEnabled);
+				InternalLeftColumnRowBox->AddSlot()
+					.HAlign(Row.WholeRowWidget.HorizontalAlignment)
+					.VAlign(Row.WholeRowWidget.VerticalAlignment)
+					.Padding(DetailWidgetConstants::LeftRowPadding)
+					[
+						Row.WholeRowWidget.Widget
+					];
+				InternalLeftColumnRowBox->AddSlot()
+					.Padding(3.0f, 0.0f)
+					.HAlign(HAlign_Right)
+					.VAlign(VAlign_Center)
+					[
+						KeyFrameButton
+					];
+				Widget = InternalLeftColumnRowBox;
+			}
 		}
 	}
-
+	else
+	{
+		// details panel layout became invalid.  This is probably a scenario where a widget is coming into view in the parent tree but some external event previous in the frame has invalidated the contents of the details panel.
+		// The next frame update of the details panel will fix it
+		Widget = SNew(SSpacer);
+	}
 
 	this->ChildSlot
 	[	
diff --git a/Engine/Source/Editor/PropertyEditor/Private/SDetailsView.cpp b/Engine/Source/Editor/PropertyEditor/Private/SDetailsView.cpp
--- a/Engine/Source/Editor/PropertyEditor/Private/SDetailsView.cpp
+++ b/Engine/Source/Editor/PropertyEditor/Private/SDetailsView.cpp
@@ -591,13 +591,15 @@ void SDetailsView::SetObjectArrayPrivate(const TArray< TWeakObjectPtr< UObject >
 	{
 		// if the object is the default metaobject for a UClass, use the UClass's name instead
 		UObject* Object = RootPropertyNodes[0]->AsObjectNode()->GetUObject(0);
-		FString ObjectName = Object->GetName();
-		if ( Object->GetClass()->GetDefaultObject() == Object )
+
+		FString ObjectName;
+		if ( Object && Object->GetClass()->GetDefaultObject() == Object )
 		{
 			ObjectName = Object->GetClass()->GetName();
 		}
-		else
+		else if( Object )
 		{
+			ObjectName = Object->GetName();
 			// Is this an actor?  If so, it might have a friendly name to display
 			const AActor* Actor = Cast<const  AActor >( Object );
 			if( Actor != nullptr)
diff --git a/Engine/Source/Editor/PropertyEditor/Private/SDetailsViewBase.cpp b/Engine/Source/Editor/PropertyEditor/Private/SDetailsViewBase.cpp
--- a/Engine/Source/Editor/PropertyEditor/Private/SDetailsViewBase.cpp
+++ b/Engine/Source/Editor/PropertyEditor/Private/SDetailsViewBase.cpp
@@ -44,9 +44,26 @@ TSharedRef<ITableRow> SDetailsViewBase::OnGenerateRowForDetailTree(TSharedRef<ID
 
 void SDetailsViewBase::SetRootExpansionStates(const bool bExpand, const bool bRecurse)
 {
-	for (auto Iter = RootTreeNodes.CreateIterator(); Iter; ++Iter)
+	if(ContainsMultipleTopLevelObjects())
 	{
-		SetNodeExpansionState(*Iter, bExpand, bRecurse);
+		FDetailNodeList Children;
+		for(auto Iter = RootTreeNodes.CreateIterator(); Iter; ++Iter)
+		{
+			Children.Reset();
+			(*Iter)->GetChildren(Children);
+
+			for(TSharedRef<class IDetailTreeNode>& Child : Children)
+			{
+				SetNodeExpansionState(Child, bExpand, bRecurse);
+			}
+		}
+	}
+	else
+	{
+		for(auto Iter = RootTreeNodes.CreateIterator(); Iter; ++Iter)
+		{
+			SetNodeExpansionState(*Iter, bExpand, bRecurse);
+		}
 	}
 }
 
diff --git a/Engine/Source/Editor/PropertyEditor/Private/UserInterface/PropertyEditor/SPropertyEditorAsset.cpp b/Engine/Source/Editor/PropertyEditor/Private/UserInterface/PropertyEditor/SPropertyEditorAsset.cpp
--- a/Engine/Source/Editor/PropertyEditor/Private/UserInterface/PropertyEditor/SPropertyEditorAsset.cpp
+++ b/Engine/Source/Editor/PropertyEditor/Private/UserInterface/PropertyEditor/SPropertyEditorAsset.cpp
@@ -190,7 +190,6 @@ void SPropertyEditorAsset::Construct( const FArguments& InArgs, const TSharedPtr
 				PropertyNode = StaticCastSharedPtr<FPropertyHandleBase>(PropertyHandle)->GetPropertyNode();
 			}
 
-
 			FObjectPropertyNode* RootObjectNode = PropertyNode->FindRootObjectItemParent();			
 			if(RootObjectNode)
 			{
@@ -1022,7 +1021,7 @@ UClass* SPropertyEditorAsset::GetObjectPropertyClass(const UProperty* Property)
 		Class = Cast<const UInterfaceProperty>(Property)->InterfaceClass;
 	}
 
-	check(Class != NULL);
+	checkf(Class != NULL, TEXT("Property (%s) is not an object or interface class"), Property ? *Property->GetFullName() : TEXT("null"));
 	return Class;
 }
 
diff --git a/Engine/Source/Editor/PropertyEditor/Private/UserInterface/PropertyEditor/SPropertyEditorClass.cpp b/Engine/Source/Editor/PropertyEditor/Private/UserInterface/PropertyEditor/SPropertyEditorClass.cpp
--- a/Engine/Source/Editor/PropertyEditor/Private/UserInterface/PropertyEditor/SPropertyEditorClass.cpp
+++ b/Engine/Source/Editor/PropertyEditor/Private/UserInterface/PropertyEditor/SPropertyEditorClass.cpp
@@ -106,6 +106,8 @@ void SPropertyEditorClass::Construct(const FArguments& InArgs, const TSharedPtr<
 		bIsBlueprintBaseOnly = Property->GetOwnerProperty()->HasMetaData(TEXT("BlueprintBaseOnly"));
 		RequiredInterface = Property->GetOwnerProperty()->GetClassMetaData(TEXT("MustImplement"));
 		bAllowNone = !(Property->PropertyFlags & CPF_NoClear);
+		bShowViewOptions = Property->GetOwnerProperty()->HasMetaData(TEXT("HideViewOptions")) ? false : true;
+		bShowTree = Property->GetOwnerProperty()->HasMetaData(TEXT("ShowTreeView"));
 	}
 	else
 	{
@@ -119,9 +121,13 @@ void SPropertyEditorClass::Construct(const FArguments& InArgs, const TSharedPtr<
 		bIsBlueprintBaseOnly = InArgs._IsBlueprintBaseOnly;
 		bAllowNone = InArgs._AllowNone;
 		bAllowOnlyPlaceable = false;
+		bShowViewOptions = InArgs._ShowViewOptions;
+		bShowTree = InArgs._ShowTree;
 
 		SelectedClass = InArgs._SelectedClass;
 		OnSetClass = InArgs._OnSetClass;
+
+
 	}
 	
 	SAssignNew(ComboButton, SComboButton)
@@ -204,6 +210,8 @@ TSharedRef<SWidget> SPropertyEditorClass::GenerateClassPicker()
 	ClassFilter->bAllowAbstract = bAllowAbstract;
 	Options.bIsBlueprintBaseOnly = bIsBlueprintBaseOnly;
 	Options.bIsPlaceableOnly = bAllowOnlyPlaceable;
+	Options.DisplayMode = bShowTree ? EClassViewerDisplayMode::TreeView : EClassViewerDisplayMode::ListView;
+	Options.bAllowViewOptions = bShowViewOptions;
 
 	FOnClassPicked OnPicked(FOnClassPicked::CreateRaw(this, &SPropertyEditorClass::OnClassPicked));
 
diff --git a/Engine/Source/Editor/PropertyEditor/Private/UserInterface/PropertyEditor/SPropertyEditorClass.h b/Engine/Source/Editor/PropertyEditor/Private/UserInterface/PropertyEditor/SPropertyEditorClass.h
--- a/Engine/Source/Editor/PropertyEditor/Private/UserInterface/PropertyEditor/SPropertyEditorClass.h
+++ b/Engine/Source/Editor/PropertyEditor/Private/UserInterface/PropertyEditor/SPropertyEditorClass.h
@@ -45,6 +45,10 @@ public:
 		SLATE_ARGUMENT(bool, AllowNone)
 		/** Attribute used to get the currently selected class (required if PropertyEditor == null) */
 		SLATE_ATTRIBUTE(const UClass*, SelectedClass)
+		/** Should we show the view options button at the bottom of the class picker?*/
+		SLATE_ARGUMENT(bool, ShowViewOptions)
+		/** Should we show the class picker in tree mode or list mode?*/
+		SLATE_ARGUMENT(bool, ShowTree)
 		/** Delegate used to set the currently selected class (required if PropertyEditor == null) */
 		SLATE_EVENT(FOnSetClass, OnSetClass)
 	SLATE_END_ARGS()
@@ -98,6 +102,10 @@ private:
 	bool bAllowNone;
 	/** Should only placeable classes be displayed? */
 	bool bAllowOnlyPlaceable;
+	/** Should we show the view options button at the bottom of the class picker?*/
+	bool bShowViewOptions;
+	/** Should we show the class picker in tree mode or list mode?*/
+	bool bShowTree;
 
 	/** Attribute used to get the currently selected class (required if PropertyEditor == null) */
 	TAttribute<const UClass*> SelectedClass;
