Fix crash when reloading a skeletal mesh actor with some vertex painting

diff --git a/Engine/Plugins/Experimental/AlembicImporter/Source/AlembicLibrary/Private/AbcImporter.cpp b/Engine/Plugins/Experimental/AlembicImporter/Source/AlembicLibrary/Private/AbcImporter.cpp
--- a/Engine/Plugins/Experimental/AlembicImporter/Source/AlembicLibrary/Private/AbcImporter.cpp
+++ b/Engine/Plugins/Experimental/AlembicImporter/Source/AlembicLibrary/Private/AbcImporter.cpp
@@ -653,6 +653,7 @@ TArray<UObject*> FAbcImporter::ImportAsSkeletalMesh(UObject* InParent, EObjectFl
 		// Forced to 1
 		LODModel.NumTexCoords = MergedMeshSample->NumUVSets;
 		SkeletalMesh->bHasVertexColors = true;
+		SkeletalMesh->VertexColorGuid = FGuid::NewGuid();
 
 		/* Bounding box according to animation */
 		SkeletalMesh->SetImportedBounds(AbcFile->GetArchiveBounds().GetBox());
diff --git a/Engine/Source/Editor/MeshPaint/Private/MeshPaintHelpers.cpp b/Engine/Source/Editor/MeshPaint/Private/MeshPaintHelpers.cpp
--- a/Engine/Source/Editor/MeshPaint/Private/MeshPaintHelpers.cpp
+++ b/Engine/Source/Editor/MeshPaint/Private/MeshPaintHelpers.cpp
@@ -838,6 +838,7 @@ void MeshPaintHelpers::FillVertexColors(UMeshComponent* MeshComponent, const FCo
 			Mesh->SetFlags(RF_Transactional);
 			Mesh->Modify();
 			Mesh->bHasVertexColors = true;
+			Mesh->VertexColorGuid = FGuid::NewGuid();
 
 			// Release the static mesh's resources.
 			Mesh->ReleaseResources();
diff --git a/Engine/Source/Editor/MeshPaint/Private/MeshPaintSkeletalMeshAdapter.cpp b/Engine/Source/Editor/MeshPaint/Private/MeshPaintSkeletalMeshAdapter.cpp
--- a/Engine/Source/Editor/MeshPaint/Private/MeshPaintSkeletalMeshAdapter.cpp
+++ b/Engine/Source/Editor/MeshPaint/Private/MeshPaintSkeletalMeshAdapter.cpp
@@ -321,6 +321,7 @@ void FMeshPaintGeometryAdapterForSkeletalMeshes::PreEdit()
 	ReferencedSkeletalMesh->Modify();
 
 	ReferencedSkeletalMesh->bHasVertexColors = true;
+	ReferencedSkeletalMesh->VertexColorGuid = FGuid::NewGuid();
 
 	// Release the static mesh's resources.
 	ReferencedSkeletalMesh->ReleaseResources();
@@ -334,6 +335,7 @@ void FMeshPaintGeometryAdapterForSkeletalMeshes::PreEdit()
 		// Mesh doesn't have a color vertex buffer yet!  We'll create one now.
 		LODData->StaticVertexBuffers.ColorVertexBuffer.InitFromSingleColor(FColor(255, 255, 255, 255), LODData->GetNumVertices());
 		ReferencedSkeletalMesh->bHasVertexColors = true;
+		ReferencedSkeletalMesh->VertexColorGuid = FGuid::NewGuid();
 		BeginInitResource(&LODData->StaticVertexBuffers.ColorVertexBuffer);
 	}
 }
diff --git a/Engine/Source/Editor/UnrealEd/Private/Fbx/FbxSkeletalMeshImport.cpp b/Engine/Source/Editor/UnrealEd/Private/Fbx/FbxSkeletalMeshImport.cpp
--- a/Engine/Source/Editor/UnrealEd/Private/Fbx/FbxSkeletalMeshImport.cpp
+++ b/Engine/Source/Editor/UnrealEd/Private/Fbx/FbxSkeletalMeshImport.cpp
@@ -1532,6 +1532,7 @@ USkeletalMesh* UnFbx::FFbxImporter::ImportSkeletalMesh(FImportSkeletalMeshArgs &
 
 	// Store whether or not this mesh has vertex colors
 	SkeletalMesh->bHasVertexColors = SkelMeshImportDataPtr->bHasVertexColors;
+	SkeletalMesh->VertexColorGuid = SkeletalMesh->bHasVertexColors ? FGuid::NewGuid() : FGuid();
 
 	FSkeletalMeshLODModel& LODModel = ImportedResource->LODModels[0];
 	
diff --git a/Engine/Source/Runtime/Engine/Classes/Engine/SkeletalMesh.h b/Engine/Source/Runtime/Engine/Classes/Engine/SkeletalMesh.h
--- a/Engine/Source/Runtime/Engine/Classes/Engine/SkeletalMesh.h
+++ b/Engine/Source/Runtime/Engine/Classes/Engine/SkeletalMesh.h
@@ -528,6 +528,12 @@ public:
 	UPROPERTY()
 	uint32 bHasVertexColors:1;
 
+#if WITH_EDITORONLY_DATA
+	/** The guid to compute the ddc key, it must be dirty when we change the vertex color. */
+	UPROPERTY()
+	FGuid VertexColorGuid;
+#endif
+
 	//caching optimization to avoid recalculating in non-editor builds
 	uint32 bHasActiveClothingAssets:1;
 
diff --git a/Engine/Source/Runtime/Engine/Private/SkeletalMesh.cpp b/Engine/Source/Runtime/Engine/Private/SkeletalMesh.cpp
--- a/Engine/Source/Runtime/Engine/Private/SkeletalMesh.cpp
+++ b/Engine/Source/Runtime/Engine/Private/SkeletalMesh.cpp
@@ -286,6 +286,7 @@ USkeletalMesh::USkeletalMesh(const FObjectInitializer& ObjectInitializer)
 	SkelMirrorFlipAxis = EAxis::Z;
 #if WITH_EDITORONLY_DATA
 	ImportedModel = MakeShareable(new FSkeletalMeshModel());
+	VertexColorGuid = FGuid();
 #endif
 
 	MinLod.Default = 0;
diff --git a/Engine/Source/Runtime/Engine/Private/SkeletalMeshMerge.cpp b/Engine/Source/Runtime/Engine/Private/SkeletalMeshMerge.cpp
--- a/Engine/Source/Runtime/Engine/Private/SkeletalMeshMerge.cpp
+++ b/Engine/Source/Runtime/Engine/Private/SkeletalMeshMerge.cpp
@@ -134,6 +134,9 @@ bool FSkeletalMeshMerge::FinalizeMesh()
 			if (SrcMesh->bHasVertexColors)
 			{
 				MergeMesh->bHasVertexColors = true;
+#if WITH_EDITORONLY_DATA
+				MergeMesh->VertexColorGuid = FGuid::NewGuid();
+#endif
 			}
 
 			FMergeMeshInfo& MeshInfo = SrcMeshInfo[MeshIdx];
diff --git a/Engine/Source/Runtime/Engine/Private/SkeletalMeshRenderData.cpp b/Engine/Source/Runtime/Engine/Private/SkeletalMeshRenderData.cpp
--- a/Engine/Source/Runtime/Engine/Private/SkeletalMeshRenderData.cpp
+++ b/Engine/Source/Runtime/Engine/Private/SkeletalMeshRenderData.cpp
@@ -47,6 +47,8 @@ static FString BuildSkeletalMeshDerivedDataKey(USkeletalMesh* SkelMesh)
 
 	KeySuffix += SkelMesh->GetImportedModel()->GetIdString();
 	KeySuffix += (SkelMesh->bUseFullPrecisionUVs || !GVertexElementTypeSupport.IsSupported(VET_Half2)) ? "1" : "0";
+	KeySuffix += SkelMesh->bHasVertexColors ? "1" : "0";
+	KeySuffix += SkelMesh->VertexColorGuid.ToString(EGuidFormats::Digits);
 
 	return FDerivedDataCacheInterface::BuildCacheKey(
 		TEXT("SKELETALMESH"),
