Fix issue with propagating vertex colors to static meshes in some cases not working correctly
#fix Using .OverrideVertexColors instead of .OverrideMapBuildData to do pointer comparison

diff --git a/Engine/Source/Editor/MeshPaint/Private/MeshPaintHelpers.cpp b/Engine/Source/Editor/MeshPaint/Private/MeshPaintHelpers.cpp
--- a/Engine/Source/Editor/MeshPaint/Private/MeshPaintHelpers.cpp
+++ b/Engine/Source/Editor/MeshPaint/Private/MeshPaintHelpers.cpp
@@ -134,7 +134,7 @@ bool MeshPaintHelpers::PropagateColorsToRawMesh(UStaticMesh* StaticMesh, int32 L
 	else
 	{
 		// If there's no raw mesh data, don't try to do any fixup here
-		if (SrcModel.IsRawMeshEmpty() || ComponentLODInfo.OverrideMapBuildData == nullptr)
+		if (SrcModel.IsRawMeshEmpty() || ComponentLODInfo.OverrideVertexColors == nullptr)
 		{
 			return false;
 		}
