Fix for possible crash if the mesh has invalid materials.
Also fixed the fact the FindMaterialIndicesUsingTexture() didn't work as advertised at all.

Seems like you'd attemp to paint all materials even when trying to only paint the ones using a particular texture.

diff --git a/Engine/Source/Editor/MeshPaint/Private/TexturePaintHelpers.cpp b/Engine/Source/Editor/MeshPaint/Private/TexturePaintHelpers.cpp
--- a/Engine/Source/Editor/MeshPaint/Private/TexturePaintHelpers.cpp
+++ b/Engine/Source/Editor/MeshPaint/Private/TexturePaintHelpers.cpp
@@ -407,8 +407,14 @@ void TexturePaintHelpers::FindMaterialIndicesUsingTexture(const UTexture* Textur
 	for (int32 MaterialIndex = 0; MaterialIndex < NumMaterials; ++MaterialIndex)
 	{
 		const UMaterialInterface* MaterialInterface = MeshComponent->GetMaterial(MaterialIndex);
-		const bool bUsesTexture = DoesMaterialUseTexture(MaterialInterface, Texture);
-		OutIndices.AddUnique(MaterialIndex);
+		if (MaterialInterface)
+		{
+			const bool bUsesTexture = DoesMaterialUseTexture(MaterialInterface, Texture);
+			if (bUsesTexture)
+			{
+				OutIndices.AddUnique(MaterialIndex);
+			}
+		}
 	}
 }
 
