Fix fill paint tool to keep in sync the staticmesh component painted vertices cache buffer.

diff --git a/Engine/Source/Editor/MeshPaint/Private/MeshPaintHelpers.cpp b/Engine/Source/Editor/MeshPaint/Private/MeshPaintHelpers.cpp
--- a/Engine/Source/Editor/MeshPaint/Private/MeshPaintHelpers.cpp
+++ b/Engine/Source/Editor/MeshPaint/Private/MeshPaintHelpers.cpp
@@ -795,6 +795,11 @@ void MeshPaintHelpers::SetInstanceColorDataForLOD(UStaticMeshComponent* MeshComp
 		// Initialize vertex buffer from given colors
 		ComponentLodInfo.OverrideVertexColors = new FColorVertexBuffer;
 		ComponentLodInfo.OverrideVertexColors->InitFromColorArray(Colors);
+		
+		//Update the cache painted vertices
+		ComponentLodInfo.PaintedVertices.Empty();
+		MeshComponent->CachePaintedDataIfNecessary();
+
 		BeginInitResource(ComponentLodInfo.OverrideVertexColors);
 	}
 }
@@ -851,6 +856,10 @@ void MeshPaintHelpers::SetInstanceColorDataForLOD(UStaticMeshComponent* MeshComp
 			}
 		}
 
+		//Update the cache painted vertices
+		ComponentLodInfo.PaintedVertices.Empty();
+		MeshComponent->CachePaintedDataIfNecessary();
+
 		BeginInitResource(ComponentLodInfo.OverrideVertexColors);
 	}
 }
@@ -1508,6 +1517,10 @@ void MeshPaintHelpers::ImportVertexColorsToStaticMeshComponent(UStaticMeshCompon
 			InstanceMeshLODInfo.OverrideVertexColors->VertexColor(VertexIndex) = PickVertexColorFromTextureData(MipData, UV, Texture, ColorMask);
 		}
 
+		//Update the cache painted vertices
+		InstanceMeshLODInfo.PaintedVertices.Empty();
+		StaticMeshComponent->CachePaintedDataIfNecessary();
+
 		BeginInitResource(InstanceMeshLODInfo.OverrideVertexColors);
 	}
 	else
