Fix mesh vertex paint LOD being clamped to 0 if actor includes a mesh component with no valid mesh data set
- Now only components that include a valid mesh asset are considered for clamping LOD
- In the case of linked Jira, mesh blueprint includes an editor-only mesh component, with no mesh asset set

diff --git a/Engine/Source/Editor/MeshPaint/Private/MeshPaintHelpers.cpp b/Engine/Source/Editor/MeshPaint/Private/MeshPaintHelpers.cpp
--- a/Engine/Source/Editor/MeshPaint/Private/MeshPaintHelpers.cpp
+++ b/Engine/Source/Editor/MeshPaint/Private/MeshPaintHelpers.cpp
@@ -1305,16 +1305,15 @@ void MeshPaintHelpers::ApplyVertexColorsToAllLODs(IMeshPaintGeometryAdapter& Geo
 	}
 }
 
-int32 MeshPaintHelpers::GetNumberOfLODs(const UMeshComponent* MeshComponent)
+bool MeshPaintHelpers::TryGetNumberOfLODs(const UMeshComponent* MeshComponent, int32& OutNumLODs)
 {
-	int32 NumLODs = 1;
-
 	if (const UStaticMeshComponent* StaticMeshComponent = Cast<UStaticMeshComponent>(MeshComponent))
 	{
 		const UStaticMesh* StaticMesh = StaticMeshComponent->GetStaticMesh();
 		if (StaticMesh != nullptr)
 		{
-			NumLODs = StaticMesh->GetNumLODs();
+			OutNumLODs = StaticMesh->GetNumLODs();
+			return true;
 		}
 	}
 	else if (const USkeletalMeshComponent* SkeletalMeshComponent = Cast<USkeletalMeshComponent>(MeshComponent))
@@ -1322,10 +1321,18 @@ int32 MeshPaintHelpers::GetNumberOfLODs(const UMeshComponent* MeshComponent)
 		const USkeletalMesh* SkeletalMesh = SkeletalMeshComponent->SkeletalMesh;
 		if (SkeletalMesh != nullptr)
 		{
-			NumLODs = SkeletalMesh->GetLODNum();
+			OutNumLODs = SkeletalMesh->GetLODNum();
+			return true;
 		}
 	}
 
+	return false;
+}
+
+int32 MeshPaintHelpers::GetNumberOfLODs(const UMeshComponent* MeshComponent)
+{
+	int32 NumLODs = 1;
+	TryGetNumberOfLODs(MeshComponent, NumLODs);
 	return NumLODs;
 }
 
diff --git a/Engine/Source/Editor/MeshPaint/Public/MeshPaintHelpers.h b/Engine/Source/Editor/MeshPaint/Public/MeshPaintHelpers.h
--- a/Engine/Source/Editor/MeshPaint/Public/MeshPaintHelpers.h
+++ b/Engine/Source/Editor/MeshPaint/Public/MeshPaintHelpers.h
@@ -107,6 +107,9 @@ public:
 
 	/** Returns the number of Mesh LODs for the given MeshComponent */
 	static int32 GetNumberOfLODs(const UMeshComponent* MeshComponent);
+
+	/** OutNumLODs is set to number of Mesh LODs for the given MeshComponent and returns true, or returns false of given mesh component has no valid LODs */
+	static bool TryGetNumberOfLODs(const UMeshComponent* MeshComponent, int32& OutNumLODs);
 	
 	/** Returns the number of Texture Coordinates for the given MeshComponent */
 	static int32 GetNumberOfUVs(const UMeshComponent* MeshComponent, int32 LODIndex);
diff --git a/Engine/Source/Editor/MeshPaintMode/Private/PaintModePainter.cpp b/Engine/Source/Editor/MeshPaintMode/Private/PaintModePainter.cpp
--- a/Engine/Source/Editor/MeshPaintMode/Private/PaintModePainter.cpp
+++ b/Engine/Source/Editor/MeshPaintMode/Private/PaintModePainter.cpp
@@ -1191,7 +1191,12 @@ int32 FPaintModePainter::GetMaxLODIndexToPaint() const
 
 	for (const UMeshComponent* MeshComponent : SelectedComponents )
 	{
-		LODMin = FMath::Min(LODMin, MeshPaintHelpers::GetNumberOfLODs(MeshComponent) - 1);
+		int32 NumMeshLODs = 0;
+		if (MeshPaintHelpers::TryGetNumberOfLODs(MeshComponent, NumMeshLODs))
+		{
+			ensure(NumMeshLODs > 0);
+			LODMin = FMath::Min(LODMin, NumMeshLODs - 1);
+		}
 	}
 
 	if (LODMin == TNumericLimits<int32>::Max())
