For static mesh vertex painting, ensure all LODs are updated when filling all LODS.
Fill now respects bPaintOnSpecificLOD for static meshes.

diff --git a/Engine/Source/Editor/MeshPaint/Private/MeshPaintHelpers.cpp b/Engine/Source/Editor/MeshPaint/Private/MeshPaintHelpers.cpp
--- a/Engine/Source/Editor/MeshPaint/Private/MeshPaintHelpers.cpp
+++ b/Engine/Source/Editor/MeshPaint/Private/MeshPaintHelpers.cpp
@@ -864,52 +864,59 @@ void MeshPaintHelpers::SetInstanceColorDataForLOD(UStaticMeshComponent* MeshComp
 	}
 }
 
-void MeshPaintHelpers::FillVertexColors(UMeshComponent* MeshComponent, const FColor FillColor, const FColor MaskColor, bool bInstanced /*= false*/)
+void MeshPaintHelpers::FillStaticMeshVertexColors(UStaticMeshComponent* MeshComponent, int32 LODIndex, const FColor FillColor, const FColor MaskColor)
 {
-	if (UStaticMeshComponent* StaticMeshComponent = Cast<UStaticMeshComponent>(MeshComponent))
+	UStaticMesh* Mesh = MeshComponent->GetStaticMesh();
+	if (Mesh)
 	{
-		if (bInstanced)
+		const int32 NumLODs = Mesh->GetNumLODs();
+		if (LODIndex < NumLODs)
 		{
-			UStaticMesh* Mesh = StaticMeshComponent->GetStaticMesh();
-			if (Mesh && Mesh->GetNumLODs() > 0)
+			if (LODIndex == -1)
 			{
-				const int32 NumLods = Mesh->GetNumLODs();
-				for (int32 LODIndex = 0; LODIndex < NumLods; ++LODIndex)
+				for (LODIndex = 0; LODIndex < NumLODs; ++LODIndex)
 				{
-					MeshPaintHelpers::SetInstanceColorDataForLOD(StaticMeshComponent, LODIndex, FillColor, MaskColor);
+					MeshPaintHelpers::SetInstanceColorDataForLOD(MeshComponent, LODIndex, FillColor, MaskColor);
 				}
 			}
+			else
+			{
+				MeshPaintHelpers::SetInstanceColorDataForLOD(MeshComponent, LODIndex, FillColor, MaskColor);
+			}
 		}
 	}
-	else if (USkeletalMeshComponent* SkeletalMeshComponent = Cast<USkeletalMeshComponent>(MeshComponent))
+}
+
+void MeshPaintHelpers::FillSkeletalMeshVertexColors(USkeletalMeshComponent* MeshComponent, int32 LODIndex, const FColor FillColor, const FColor MaskColor)
+{
+	TUniquePtr< FSkinnedMeshComponentRecreateRenderStateContext > RecreateRenderStateContext;
+	USkeletalMesh* Mesh = MeshComponent->SkeletalMesh;
+	if (Mesh)
 	{
-		TUniquePtr< FSkinnedMeshComponentRecreateRenderStateContext > RecreateRenderStateContext;
-		USkeletalMesh* Mesh = SkeletalMeshComponent->SkeletalMesh;
-		if (Mesh)
-		{
-			// Dirty the mesh
-			Mesh->SetFlags(RF_Transactional);
-			Mesh->Modify();
-			Mesh->bHasVertexColors = true;
-			Mesh->VertexColorGuid = FGuid::NewGuid();
+		// Dirty the mesh
+		Mesh->SetFlags(RF_Transactional);
+		Mesh->Modify();
+		Mesh->bHasVertexColors = true;
+		Mesh->VertexColorGuid = FGuid::NewGuid();
 
-			// Release the static mesh's resources.
-			Mesh->ReleaseResources();
+		// Release the static mesh's resources.
+		Mesh->ReleaseResources();
 
-			// Flush the resource release commands to the rendering thread to ensure that the build doesn't occur while a resource is still
-			// allocated, and potentially accessing the UStaticMesh.
-			Mesh->ReleaseResourcesFence.Wait();
+		// Flush the resource release commands to the rendering thread to ensure that the build doesn't occur while a resource is still
+		// allocated, and potentially accessing the UStaticMesh.
+		Mesh->ReleaseResourcesFence.Wait();
 
-			if (Mesh->GetLODNum() > 0)
+		const int32 NumLODs = Mesh->GetLODNum();
+		if (NumLODs > 0)
+		{
+			RecreateRenderStateContext = MakeUnique<FSkinnedMeshComponentRecreateRenderStateContext>(Mesh);
+			// TODO: Apply to LODIndex only (or all if set to -1). This requires some extra refactoring
+			// because currently all LOD data is being released above.
+			for (LODIndex = 0; LODIndex < NumLODs; ++LODIndex)
 			{
-				RecreateRenderStateContext = MakeUnique<FSkinnedMeshComponentRecreateRenderStateContext>(Mesh);
-				const int32 NumLods = Mesh->GetLODNum();
-				for (int32 LODIndex = 0; LODIndex < NumLods; ++LODIndex)
-				{
-					MeshPaintHelpers::SetColorDataForLOD(Mesh, LODIndex, FillColor, MaskColor);
-				}
-				Mesh->InitResources();
+				MeshPaintHelpers::SetColorDataForLOD(Mesh, LODIndex, FillColor, MaskColor);
 			}
+			Mesh->InitResources();
 		}
 	}
 }
diff --git a/Engine/Source/Editor/MeshPaint/Private/MeshPaintStaticMeshAdapter.cpp b/Engine/Source/Editor/MeshPaint/Private/MeshPaintStaticMeshAdapter.cpp
--- a/Engine/Source/Editor/MeshPaint/Private/MeshPaintStaticMeshAdapter.cpp
+++ b/Engine/Source/Editor/MeshPaint/Private/MeshPaintStaticMeshAdapter.cpp
@@ -118,8 +118,13 @@ void FMeshPaintGeometryAdapterForStaticMeshes::PostEdit()
 		// only need to detach our mesh component
 		FComponentReregisterContext ComponentReregisterContext(StaticMeshComponent);
 
-		FStaticMeshComponentLODInfo* InstanceMeshLODInfo = &StaticMeshComponent->LODData[MeshLODIndex];
-		BeginInitResource(InstanceMeshLODInfo->OverrideVertexColors);
+		// If LOD is 0, post-edit all LODs. There's currently no way to tell from here
+		// if VertexPaintSettings.bPaintOnSpecificLOD is set to true or not.
+		const int32 MaxLOD = (MeshLODIndex == 0) ? StaticMeshComponent->LODData.Num() : (MeshLODIndex + 1);
+		for (int32 Index = MeshLODIndex; Index < MaxLOD; ++Index)
+		{
+			BeginInitResource(StaticMeshComponent->LODData[Index].OverrideVertexColors);
+		}
 	}
 	else
 	{
@@ -364,7 +369,6 @@ void FMeshPaintGeometryAdapterForStaticMeshes::PreEdit()
 {
 	const bool bUsingInstancedVertexColors = true; // Currently we are only painting to instances 
 	UStaticMesh* StaticMesh = ReferencedStaticMesh;
-	FStaticMeshComponentLODInfo* InstanceMeshLODInfo = NULL;
 	if (bUsingInstancedVertexColors)
 	{
 		// Mark the mesh component as modified
@@ -375,38 +379,45 @@ void FMeshPaintGeometryAdapterForStaticMeshes::PreEdit()
 		// Ensure LODData has enough entries in it, free not required.
 		StaticMeshComponent->SetLODDataCount(MeshLODIndex + 1, StaticMeshComponent->LODData.Num());
 
-		InstanceMeshLODInfo = &StaticMeshComponent->LODData[MeshLODIndex];
-
-		// Destroy the instance vertex  color array if it doesn't fit
-		if (InstanceMeshLODInfo->OverrideVertexColors
-			&& InstanceMeshLODInfo->OverrideVertexColors->GetNumVertices() != LODModel->GetNumVertices())
+		// If LOD is 0, pre-edit all LODs. There's currently no way to tell from here
+		// if VertexPaintSettings.bPaintOnSpecificLOD is set to true or not.
+		const int32 MaxLOD = (MeshLODIndex == 0) ? StaticMeshComponent->LODData.Num() : (MeshLODIndex + 1);
+		for (int32 Index = MeshLODIndex; Index < MaxLOD; ++Index)
 		{
-			InstanceMeshLODInfo->ReleaseOverrideVertexColorsAndBlock();
-		}
+			FStaticMeshComponentLODInfo& InstanceMeshLODInfo = StaticMeshComponent->LODData[Index];
+			FStaticMeshLODResources& LODResource = StaticMesh->RenderData->LODResources[Index];
 
-		if (InstanceMeshLODInfo->OverrideVertexColors)
-		{
-			// Destroy the cached paint data every paint. Painting redefines the source data.
-			InstanceMeshLODInfo->PaintedVertices.Empty();
-			InstanceMeshLODInfo->BeginReleaseOverrideVertexColors();
-			FlushRenderingCommands();
-		}
-		else
-		{
-			// Setup the instance vertex color array if we don't have one yet
-			InstanceMeshLODInfo->OverrideVertexColors = new FColorVertexBuffer;
+			// Destroy the instance vertex  color array if it doesn't fit
+			if (InstanceMeshLODInfo.OverrideVertexColors
+				&& InstanceMeshLODInfo.OverrideVertexColors->GetNumVertices() != LODResource.GetNumVertices())
+			{
+				InstanceMeshLODInfo.ReleaseOverrideVertexColorsAndBlock();
+			}
 
-			if ((int32)LODModel->VertexBuffers.ColorVertexBuffer.GetNumVertices() >= LODModel->GetNumVertices())
+			if (InstanceMeshLODInfo.OverrideVertexColors)
 			{
-				// copy mesh vertex colors to the instance ones
-				InstanceMeshLODInfo->OverrideVertexColors->InitFromColorArray(&LODModel->VertexBuffers.ColorVertexBuffer.VertexColor(0), LODModel->GetNumVertices());
+				// Destroy the cached paint data every paint. Painting redefines the source data.
+				InstanceMeshLODInfo.PaintedVertices.Empty();
+				InstanceMeshLODInfo.BeginReleaseOverrideVertexColors();
+				FlushRenderingCommands();
 			}
 			else
 			{
-				// Original mesh didn't have any colors, so just use a default color
-				InstanceMeshLODInfo->OverrideVertexColors->InitFromSingleColor(FColor::White, LODModel->GetNumVertices());
-			}
+				// Setup the instance vertex color array if we don't have one yet
+				InstanceMeshLODInfo.OverrideVertexColors = new FColorVertexBuffer;
+
+				if ((int32)LODResource.VertexBuffers.ColorVertexBuffer.GetNumVertices() >= LODResource.GetNumVertices())
+				{
+					// copy mesh vertex colors to the instance ones
+					InstanceMeshLODInfo.OverrideVertexColors->InitFromColorArray(&LODResource.VertexBuffers.ColorVertexBuffer.VertexColor(0), LODResource.GetNumVertices());
+				}
+				else
+				{
+					// Original mesh didn't have any colors, so just use a default color
+					InstanceMeshLODInfo.OverrideVertexColors->InitFromSingleColor(FColor::White, LODResource.GetNumVertices());
+				}
 
+			}
 		}
 		// See if the component has to cache its mesh vertex positions associated with override colors
 		StaticMeshComponent->CachePaintedDataIfNecessary();
diff --git a/Engine/Source/Editor/MeshPaint/Public/MeshPaintHelpers.h b/Engine/Source/Editor/MeshPaint/Public/MeshPaintHelpers.h
--- a/Engine/Source/Editor/MeshPaint/Public/MeshPaintHelpers.h
+++ b/Engine/Source/Editor/MeshPaint/Public/MeshPaintHelpers.h
@@ -71,7 +71,8 @@ public:
 	static void SetInstanceColorDataForLOD(UStaticMeshComponent* MeshComponent, int32 LODIndex, const FColor FillColor, const FColor MaskColor);
 	
 	/** Fills all vertex colors for all LODs found in the given mesh component with Fill Color */
-	static void FillVertexColors(UMeshComponent* MeshComponent, const FColor FillColor, const FColor MaskColor, bool bInstanced = false);
+	static void FillStaticMeshVertexColors(UStaticMeshComponent* MeshComponent, int32 LODIndex, const FColor FillColor, const FColor MaskColor);
+	static void FillSkeletalMeshVertexColors(USkeletalMeshComponent* MeshComponent, int32 LODIndex, const FColor FillColor, const FColor MaskColor);
 	
 	/** Sets all vertex colors for a specific LOD level in the SkeletalMesh to FillColor */
 	static void SetColorDataForLOD(USkeletalMesh* SkeletalMesh, int32 LODIndex, const FColor FillColor, const FColor MaskColor);
diff --git a/Engine/Source/Editor/MeshPaintMode/Private/PaintModePainter.cpp b/Engine/Source/Editor/MeshPaintMode/Private/PaintModePainter.cpp
--- a/Engine/Source/Editor/MeshPaintMode/Private/PaintModePainter.cpp
+++ b/Engine/Source/Editor/MeshPaintMode/Private/PaintModePainter.cpp
@@ -2049,7 +2049,14 @@ void FPaintModePainter::FillWithVertexColor()
 			(*MeshAdapter)->PreEdit();
 		}
 		
-		MeshPaintHelpers::FillVertexColors(Component, FillColor, MaskColor, true);
+		if (Component->IsA<UStaticMeshComponent>())
+		{
+			MeshPaintHelpers::FillStaticMeshVertexColors(Cast<UStaticMeshComponent>(Component), PaintSettings->VertexPaintSettings.bPaintOnSpecificLOD ? PaintSettings->VertexPaintSettings.LODIndex : -1, FillColor, MaskColor);
+		}
+		else if (Component->IsA<USkeletalMeshComponent>())
+		{
+			MeshPaintHelpers::FillSkeletalMeshVertexColors(Cast<USkeletalMeshComponent>(Component), PaintSettings->VertexPaintSettings.bPaintOnSpecificLOD ? PaintSettings->VertexPaintSettings.LODIndex : -1, FillColor, MaskColor);
+		}
 
 		if (MeshAdapter)
 		{
