Deprecate FRawMesh - Mesh paint helper

diff --git a/Engine/Source/Editor/MeshPaint/MeshPaint.Build.cs b/Engine/Source/Editor/MeshPaint/MeshPaint.Build.cs
--- a/Engine/Source/Editor/MeshPaint/MeshPaint.Build.cs
+++ b/Engine/Source/Editor/MeshPaint/MeshPaint.Build.cs
@@ -28,7 +28,7 @@ public class MeshPaint : ModuleRules
 				"SlateCore",
                 "EditorStyle",
                 "UnrealEd",
-                "RawMesh",
+                "MeshDescription",
                 "SourceControl",
                 "ViewportInteraction",
                 "VREditor",
diff --git a/Engine/Source/Editor/MeshPaint/Private/MeshPaintHelpers.cpp b/Engine/Source/Editor/MeshPaint/Private/MeshPaintHelpers.cpp
--- a/Engine/Source/Editor/MeshPaint/Private/MeshPaintHelpers.cpp
+++ b/Engine/Source/Editor/MeshPaint/Private/MeshPaintHelpers.cpp
@@ -14,7 +14,11 @@
 #include "Engine/SkeletalMesh.h"
 #include "Engine/Texture2D.h"
 #include "StaticMeshResources.h"
-#include "RawMesh.h"
+
+#include "MeshDescription.h"
+#include "MeshAttributes.h"
+#include "MeshAttributeArray.h"
+
 #include "Rendering/SkeletalMeshRenderData.h"
 
 #include "GenericOctree.h"
@@ -106,46 +110,60 @@ bool MeshPaintHelpers::PropagateColorsToRawMesh(UStaticMesh* StaticMesh, int32 L
 	if (RenderData.WedgeMap.Num() > 0 && ColorVertexBuffer.GetNumVertices() == RenderModel.GetNumVertices())
 	{
 		// Use the wedge map if it is available as it is lossless.
-		FRawMesh RawMesh;
-		SrcModel.LoadRawMesh(RawMesh);
+		FMeshDescription* RawMeshPtr = StaticMesh->GetOriginalMeshDescription(LODIndex);
+
+		if (RawMeshPtr == nullptr)
+		{
+			//Cannot propagate to a generated LOD, the generated LOD use the source LOD vertex painting. 
+			return false;
+		}
 
-		int32 NumWedges = RawMesh.WedgeIndices.Num();
+		FMeshDescription& RawMesh = *RawMeshPtr;
+		int32 NumWedges = RawMesh.VertexInstances().Num();
 		if (RenderData.WedgeMap.Num() == NumWedges)
 		{
-			int32 NumExistingColors = RawMesh.WedgeColors.Num();
-			if (NumExistingColors < NumWedges)
+			FStaticMeshDescriptionAttributeGetter AttributeGetter(&RawMesh);
+			TVertexInstanceAttributesRef<FVector4> Colors = AttributeGetter.GetColors();
+			int32 VertexInstanceIndex = 0;
+			for (const FVertexInstanceID VertexInstanceID : RawMesh.VertexInstances().GetElementIDs())
 			{
-				RawMesh.WedgeColors.AddUninitialized(NumWedges - NumExistingColors);
-			}
-			for (int32 i = 0; i < NumWedges; ++i)
-			{
-				FColor WedgeColor = FColor::White;
-				int32 Index = RenderData.WedgeMap[i];
+				FLinearColor WedgeColor = FLinearColor::White;
+				int32 Index = RenderData.WedgeMap[VertexInstanceIndex];
 				if (Index != INDEX_NONE)
 				{
-					WedgeColor = ColorVertexBuffer.VertexColor(Index);
+					WedgeColor = FLinearColor(ColorVertexBuffer.VertexColor(Index));
 				}
-				RawMesh.WedgeColors[i] = WedgeColor;
+				Colors[VertexInstanceID] = WedgeColor;
+				VertexInstanceIndex++;
 			}
-			SrcModel.SaveRawMesh(RawMesh);
+			StaticMesh->CommitOriginalMeshDescription(LODIndex);
 			bPropagatedColors = true;
 		}
 	}
 	else
 	{
+		FMeshDescription* SrcModelMeshDescription = StaticMesh->GetOriginalMeshDescription(LODIndex);
 		// If there's no raw mesh data, don't try to do any fixup here
-		if (SrcModel.IsRawMeshEmpty() || ComponentLODInfo.OverrideVertexColors == nullptr)
+		if (SrcModelMeshDescription == nullptr || ComponentLODInfo.OverrideVertexColors == nullptr)
 		{
 			return false;
 		}
 
 		// Fall back to mapping based on position.
-		FRawMesh RawMesh;
-		SrcModel.LoadRawMesh(RawMesh);
-
+		FStaticMeshDescriptionAttributeGetter AttributeGetter(SrcModelMeshDescription);
+		TVertexAttributesConstRef<FVector> VertexPositions = AttributeGetter.GetPositionsConst();
+		TVertexInstanceAttributesRef<FVector4> Colors = AttributeGetter.GetColors();
 		TArray<FColor> NewVertexColors;
 		FPositionVertexBuffer TempPositionVertexBuffer;
-		TempPositionVertexBuffer.Init(RawMesh.VertexPositions);
+		int32 NumVertex = SrcModelMeshDescription->Vertices().Num();
+		TArray<FVector> VertexPositionsDup;
+		VertexPositionsDup.AddZeroed(NumVertex);
+		int32 VertexIndex = 0;
+		for (const FVertexID VertexID : SrcModelMeshDescription->Vertices().GetElementIDs())
+		{
+			VertexPositionsDup[VertexIndex] = VertexPositions[VertexID];
+		}
+		TempPositionVertexBuffer.Init(VertexPositionsDup);
 		RemapPaintedVertexColors(
 			ComponentLODInfo.PaintedVertices,
 			ComponentLODInfo.OverrideVertexColors,
@@ -155,17 +173,14 @@ bool MeshPaintHelpers::PropagateColorsToRawMesh(UStaticMesh* StaticMesh, int32 L
 			/*OptionalVertexBuffer=*/ nullptr,
 			NewVertexColors
 		);
-		if (NewVertexColors.Num() == RawMesh.VertexPositions.Num())
+		if (NewVertexColors.Num() == NumVertex)
 		{
-			int32 NumWedges = RawMesh.WedgeIndices.Num();
-			RawMesh.WedgeColors.Empty(NumWedges);
-			RawMesh.WedgeColors.AddZeroed(NumWedges);
-			for (int32 i = 0; i < NumWedges; ++i)
+			for (const FVertexInstanceID VertexInstanceID : SrcModelMeshDescription->VertexInstances().GetElementIDs())
 			{
-				int32 Index = RawMesh.WedgeIndices[i];
-				RawMesh.WedgeColors[i] = NewVertexColors[Index];
+				const FVertexID VertexID = SrcModelMeshDescription->GetVertexInstanceVertex(VertexInstanceID);
+				Colors[VertexInstanceID] = FLinearColor(NewVertexColors[VertexID.GetValue()]);
 			}
-			SrcModel.SaveRawMesh(RawMesh);
+			StaticMesh->CommitOriginalMeshDescription(LODIndex);
 			bPropagatedColors = true;
 		}
 	}
