Apply mask to fill color when filling LOD color data

diff --git a/Engine/Source/Editor/MeshPaint/Private/MeshPaintHelpers.cpp b/Engine/Source/Editor/MeshPaint/Private/MeshPaintHelpers.cpp
--- a/Engine/Source/Editor/MeshPaint/Private/MeshPaintHelpers.cpp
+++ b/Engine/Source/Editor/MeshPaint/Private/MeshPaintHelpers.cpp
@@ -820,7 +820,9 @@ void MeshPaintHelpers::SetInstanceColorDataForLOD(UStaticMeshComponent* MeshComp
 			{
 				// Initialize vertex buffer from given color
 				ComponentLodInfo.OverrideVertexColors = new FColorVertexBuffer;
-				ComponentLodInfo.OverrideVertexColors->InitFromSingleColor(MaskColor, RenderData.GetNumVertices());
+				FColor NewFillColor(EForceInit::ForceInitToZero);
+				ApplyFillWithMask(NewFillColor, MaskColor, FillColor);
+				ComponentLodInfo.OverrideVertexColors->InitFromSingleColor(NewFillColor, RenderData.GetNumVertices());
 			}
 		}
 
