Second call to InitRHI from BeginInitResource (already been called in MeshPaintHelpers::SetInstanceColorDataForLOD()) will be ignored due to bInitialized flag on the FRenderResource set in the render thread - race condition on VertexColorFromTextureData.

Don't use one of the generic functions as they call the begin init resource and seem to be designed for external use - simpler just to call the init from single color function in a similar fashon to the main if part of the condtional clause - keeps the code paths doing similar work and avoids double init calls.

diff --git a/Engine/Source/Editor/MeshPaint/Private/MeshPaintHelpers.cpp b/Engine/Source/Editor/MeshPaint/Private/MeshPaintHelpers.cpp
--- a/Engine/Source/Editor/MeshPaint/Private/MeshPaintHelpers.cpp
+++ b/Engine/Source/Editor/MeshPaint/Private/MeshPaintHelpers.cpp
@@ -1467,7 +1467,12 @@ void MeshPaintHelpers::ImportVertexColorsToStaticMeshComponent(UStaticMeshCompon
 		else
 		{
 			// Original mesh didn't have any colors, so just use a default color
-			MeshPaintHelpers::SetInstanceColorDataForLOD(StaticMeshComponent, ImportLOD, FColor::White, FColor::White);
+			InstanceMeshLODInfo.OverrideVertexColors->InitFromSingleColor(FColor::White, LODModel.GetNumVertices());
+		}
+		
+		if (ImportLOD > 0)
+		{
+			StaticMeshComponent->bCustomOverrideVertexColorPerLOD = true;
 		}
 
 		const int32 UVIndex = Options->UVIndex;
