RL - Move RecreateRenderState and Context to post edit chagne as they will rebuild MeshDrawCommands and those have not yet have properly set up VertexFactories.

diff --git a/Engine/Source/Editor/MeshPaint/Private/MeshPaintStaticMeshAdapter.cpp b/Engine/Source/Editor/MeshPaint/Private/MeshPaintStaticMeshAdapter.cpp
--- a/Engine/Source/Editor/MeshPaint/Private/MeshPaintStaticMeshAdapter.cpp
+++ b/Engine/Source/Editor/MeshPaint/Private/MeshPaintStaticMeshAdapter.cpp
@@ -108,12 +108,16 @@ void FMeshPaintGeometryAdapterForStaticMeshes::PostEdit()
 	const bool bUnbuildLighting = false;
 
 	// Recreate all component states using the referenced static mesh
-	TUniquePtr<FStaticMeshComponentRecreateRenderStateContext> RecreateRenderStateContext = MakeUnique<FStaticMeshComponentRecreateRenderStateContext>(ReferencedStaticMesh, bUnbuildLighting);
-	const bool bUsingInstancedVertexColors = true;
-
+	FStaticMeshComponentRecreateRenderStateContext RecreateRenderStateContext(ReferencedStaticMesh, bUnbuildLighting);
+	
+	const bool bUsingInstancedVertexColors = true; // Currently we are only painting to instances 
 	// Update gpu resource data 
 	if (bUsingInstancedVertexColors)
 	{
+		// We're only changing instanced vertices on this specific mesh component, so we
+		// only need to detach our mesh component
+		FComponentReregisterContext ComponentReregisterContext(StaticMeshComponent);
+
 		FStaticMeshComponentLODInfo* InstanceMeshLODInfo = &StaticMeshComponent->LODData[MeshLODIndex];
 		BeginInitResource(InstanceMeshLODInfo->OverrideVertexColors);
 	}
@@ -361,14 +365,8 @@ void FMeshPaintGeometryAdapterForStaticMeshes::PreEdit()
 	const bool bUsingInstancedVertexColors = true; // Currently we are only painting to instances 
 	UStaticMesh* StaticMesh = ReferencedStaticMesh;
 	FStaticMeshComponentLODInfo* InstanceMeshLODInfo = NULL;
-	TUniquePtr< FStaticMeshComponentRecreateRenderStateContext > RecreateRenderStateContext;
-	TUniquePtr< FComponentReregisterContext > ComponentReregisterContext;
 	if (bUsingInstancedVertexColors)
 	{
-		// We're only changing instanced vertices on this specific mesh component, so we
-		// only need to detach our mesh component
-		ComponentReregisterContext = MakeUnique<FComponentReregisterContext>(StaticMeshComponent);
-
 		// Mark the mesh component as modified
 		StaticMeshComponent->SetFlags(RF_Transactional);
 		StaticMeshComponent->Modify();
@@ -386,14 +384,10 @@ void FMeshPaintGeometryAdapterForStaticMeshes::PreEdit()
 			InstanceMeshLODInfo->ReleaseOverrideVertexColorsAndBlock();
 		}
 
-		// Destroy the cached paint data every paint. Painting redefines the source data.
 		if (InstanceMeshLODInfo->OverrideVertexColors)
 		{
+			// Destroy the cached paint data every paint. Painting redefines the source data.
 			InstanceMeshLODInfo->PaintedVertices.Empty();
-		}
-
-		if (InstanceMeshLODInfo->OverrideVertexColors)
-		{
 			InstanceMeshLODInfo->BeginReleaseOverrideVertexColors();
 			FlushRenderingCommands();
 		}
@@ -420,10 +414,6 @@ void FMeshPaintGeometryAdapterForStaticMeshes::PreEdit()
 	}
 	else
 	{
-		// We're changing the mesh itself, so ALL static mesh components in the scene will need
-		// to be unregistered for this (and reregistered afterwards.)
-		RecreateRenderStateContext = MakeUnique<FStaticMeshComponentRecreateRenderStateContext>(StaticMesh);
-
 		// Dirty the mesh
 		StaticMesh->SetFlags(RF_Transactional);
 		StaticMesh->Modify();
