The mesh painting settings Enable Brush Flow, Ignore Back-Facing and Color View Mode are now persisted between editor sessions.

Clamping for settings is now read from the UProperty ClampMin/ClampMax values.

diff --git a/Engine/Source/Editor/MeshPaint/Private/MeshPaintSettings.cpp b/Engine/Source/Editor/MeshPaint/Private/MeshPaintSettings.cpp
--- a/Engine/Source/Editor/MeshPaint/Private/MeshPaintSettings.cpp
+++ b/Engine/Source/Editor/MeshPaint/Private/MeshPaintSettings.cpp
@@ -13,16 +13,50 @@ UPaintBrushSettings::UPaintBrushSettings(const FObjectInitializer& ObjectInitial
 	bOnlyFrontFacingTriangles(true),
 	ColorViewMode(EMeshPaintColorViewMode::Normal)	
 {
-	BrushRadiusMin = 0.01f, BrushRadiusMax = 250000.0f;
+	const FName ClampMin("ClampMin");
+	const FName ClampMax("ClampMax");
 
-	GConfig->GetFloat(TEXT("MeshPaintEdit"), TEXT("DefaultBrushRadius"), BrushRadius, GEditorPerProjectIni);
-	BrushRadius = FMath::Clamp(BrushRadius, BrushRadiusMin, BrushRadiusMax);
+	{
+		const UProperty* BrushRadiusProperty = UPaintBrushSettings::StaticClass()->FindPropertyByName(GET_MEMBER_NAME_CHECKED(UPaintBrushSettings, BrushRadius));
+		BrushRadiusMin = BrushRadiusProperty->GetFLOATMetaData(ClampMin);
+		BrushRadiusMax = BrushRadiusProperty->GetFLOATMetaData(ClampMax);
+
+		GConfig->GetFloat(TEXT("MeshPaintEdit"), TEXT("DefaultBrushRadius"), BrushRadius, GEditorPerProjectIni);
+		BrushRadius = FMath::Clamp(BrushRadius, BrushRadiusMin, BrushRadiusMax);
+	}
 
-	GConfig->GetFloat(TEXT("MeshPaintEdit"), TEXT("DefaultBrushStrength"), BrushStrength, GEditorPerProjectIni);
-	BrushStrength = FMath::Clamp(BrushRadius, 0.f, 1.f);
+	{
+		const UProperty* BrushStrengthProperty = UPaintBrushSettings::StaticClass()->FindPropertyByName(GET_MEMBER_NAME_CHECKED(UPaintBrushSettings, BrushStrength));
+		float BrushStrengthMin = BrushStrengthProperty->GetFLOATMetaData(ClampMin);
+		float BrushStrengthMax = BrushStrengthProperty->GetFLOATMetaData(ClampMax);
 
-	GConfig->GetFloat(TEXT("MeshPaintEdit"), TEXT("DefaultBrushFalloff"), BrushFalloffAmount, GEditorPerProjectIni);
-	BrushFalloffAmount = FMath::Clamp(BrushFalloffAmount, 0.f, 1.f);
+		GConfig->GetFloat(TEXT("MeshPaintEdit"), TEXT("DefaultBrushStrength"), BrushStrength, GEditorPerProjectIni);
+		BrushStrength = FMath::Clamp(BrushStrength, BrushStrengthMin, BrushStrengthMax);
+	}
+
+	{
+		const UProperty* BrushFalloffProperty = UPaintBrushSettings::StaticClass()->FindPropertyByName(GET_MEMBER_NAME_CHECKED(UPaintBrushSettings, BrushFalloffAmount));
+		float BrushFalloffMin = BrushFalloffProperty->GetFLOATMetaData(ClampMin);
+		float BrushFalloffMax = BrushFalloffProperty->GetFLOATMetaData(ClampMax);
+
+		GConfig->GetFloat(TEXT("MeshPaintEdit"), TEXT("DefaultBrushFalloff"), BrushFalloffAmount, GEditorPerProjectIni);
+		BrushFalloffAmount = FMath::Clamp(BrushFalloffAmount, BrushFalloffMin, BrushFalloffMax);
+	}
+
+	GConfig->GetBool(TEXT("MeshPaintEdit"), TEXT("IgnoreBackFacing"), bOnlyFrontFacingTriangles, GEditorPerProjectIni);
+
+	GConfig->GetBool(TEXT("MeshPaintEdit"), TEXT("EnableBrushFlow"), bEnableFlow, GEditorPerProjectIni);
+
+	FString ColorViewModeString;
+	if (GConfig->GetString(TEXT("MeshPaintEdit"), TEXT("ColorViewMode"), ColorViewModeString, GEditorPerProjectIni))
+	{
+		const UEnum* ColorViewModeEnum = StaticEnum<EMeshPaintColorViewMode>();
+		int64 Value = ColorViewModeEnum->GetValueByNameString(ColorViewModeString);
+		if (Value != INDEX_NONE)
+		{
+			ColorViewMode = (EMeshPaintColorViewMode) Value;
+		}
+	}
 }
 
 UPaintBrushSettings::~UPaintBrushSettings()
@@ -49,18 +83,36 @@ void UPaintBrushSettings::SetBrushFalloff(float InFalloff)
 
 void UPaintBrushSettings::PostEditChangeProperty(struct FPropertyChangedEvent& PropertyChangedEvent)
 {
-	if (PropertyChangedEvent.Property && PropertyChangedEvent.Property->GetFName() == GET_MEMBER_NAME_CHECKED(UPaintBrushSettings, BrushRadius) && PropertyChangedEvent.ChangeType != EPropertyChangeType::Interactive)
+	if (PropertyChangedEvent.Property == nullptr || PropertyChangedEvent.ChangeType == EPropertyChangeType::Interactive)
 	{
-		GConfig->SetFloat(TEXT("MeshPaintEdit"), TEXT("DefaultBrushRadius"), BrushRadius, GEditorPerProjectIni);
+		return;
 	}
 
-	if (PropertyChangedEvent.Property && PropertyChangedEvent.Property->GetFName() == GET_MEMBER_NAME_CHECKED(UPaintBrushSettings, BrushStrength) && PropertyChangedEvent.ChangeType != EPropertyChangeType::Interactive)
+	FName PropertyName = PropertyChangedEvent.Property->GetFName();
+	if (PropertyName == GET_MEMBER_NAME_CHECKED(UPaintBrushSettings, BrushRadius))
+	{
+		GConfig->SetFloat(TEXT("MeshPaintEdit"), TEXT("DefaultBrushRadius"), BrushRadius, GEditorPerProjectIni);
+	}
+	else if (PropertyName == GET_MEMBER_NAME_CHECKED(UPaintBrushSettings, BrushStrength))
 	{
 		GConfig->SetFloat(TEXT("MeshPaintEdit"), TEXT("DefaultBrushStrength"), BrushStrength, GEditorPerProjectIni);
 	}
-
-	if (PropertyChangedEvent.Property && PropertyChangedEvent.Property->GetFName() == GET_MEMBER_NAME_CHECKED(UPaintBrushSettings, BrushFalloffAmount) && PropertyChangedEvent.ChangeType != EPropertyChangeType::Interactive)
+	else if (PropertyName == GET_MEMBER_NAME_CHECKED(UPaintBrushSettings, BrushFalloffAmount))
 	{
 		GConfig->SetFloat(TEXT("MeshPaintEdit"), TEXT("DefaultBrushFalloff"), BrushFalloffAmount, GEditorPerProjectIni);
 	}
+	else if (PropertyName == GET_MEMBER_NAME_CHECKED(UPaintBrushSettings, bOnlyFrontFacingTriangles))
+	{
+		GConfig->SetBool(TEXT("MeshPaintEdit"), TEXT("IgnoreBackFacing"), bOnlyFrontFacingTriangles, GEditorPerProjectIni);
+	}
+	else if (PropertyName == GET_MEMBER_NAME_CHECKED(UPaintBrushSettings, bEnableFlow))
+	{
+		GConfig->SetBool(TEXT("MeshPaintEdit"), TEXT("EnableBrushFlow"), bEnableFlow, GEditorPerProjectIni);
+	}
+	else if (PropertyName == GET_MEMBER_NAME_CHECKED(UPaintBrushSettings, ColorViewMode))
+	{
+		const UEnum* ColorViewModeEnum = StaticEnum<EMeshPaintColorViewMode>();
+		FString ColorViewModeString = ColorViewModeEnum->GetNameStringByValue((int64) ColorViewMode);
+		GConfig->SetString(TEXT("MeshPaintEdit"), TEXT("ColorViewMode"), *ColorViewModeString, GEditorPerProjectIni);
+	}
 }
\ No newline at end of file
diff --git a/Engine/Source/Editor/MeshPaint/Public/MeshPaintSettings.h b/Engine/Source/Editor/MeshPaint/Public/MeshPaintSettings.h
--- a/Engine/Source/Editor/MeshPaint/Public/MeshPaintSettings.h
+++ b/Engine/Source/Editor/MeshPaint/Public/MeshPaintSettings.h
@@ -71,7 +71,7 @@ public:
 	bool bEnableFlow;
 
 	/** Whether back-facing triangles should be ignored */
-	UPROPERTY(EditAnywhere, Category = Brush, meta = (DisplayName = "Ignore back-facing"))
+	UPROPERTY(EditAnywhere, Category = Brush, meta = (DisplayName = "Ignore Back-Facing"))
 	bool bOnlyFrontFacingTriangles;
 	
 	/** Color view mode used to display Vertex Colors */
