Fix crash when Mesh painting SplineMeshComponent

#rnx

diff --git a/Engine/Source/Editor/MeshPaint/Private/MeshPaintStaticMeshAdapter.cpp b/Engine/Source/Editor/MeshPaint/Private/MeshPaintStaticMeshAdapter.cpp
--- a/Engine/Source/Editor/MeshPaint/Private/MeshPaintStaticMeshAdapter.cpp
+++ b/Engine/Source/Editor/MeshPaint/Private/MeshPaintStaticMeshAdapter.cpp
@@ -375,13 +375,14 @@ void FMeshPaintGeometryAdapterForStaticMeshes::PreEdit()
 		StaticMeshComponent->SetFlags(RF_Transactional);
 		StaticMeshComponent->Modify();
 		StaticMeshComponent->bCustomOverrideVertexColorPerLOD = (MeshLODIndex > 0);
-
+				
+		const int32 MaxLOD = (MeshLODIndex == 0) ? StaticMesh->RenderData->LODResources.Num() : (MeshLODIndex + 1);
 		// Ensure LODData has enough entries in it, free not required.
 		StaticMeshComponent->SetLODDataCount(MeshLODIndex + 1, StaticMeshComponent->LODData.Num());
 
 		// If LOD is 0, pre-edit all LODs. There's currently no way to tell from here
 		// if VertexPaintSettings.bPaintOnSpecificLOD is set to true or not.
-		const int32 MaxLOD = (MeshLODIndex == 0) ? StaticMeshComponent->LODData.Num() : (MeshLODIndex + 1);
+		
 		for (int32 Index = MeshLODIndex; Index < MaxLOD; ++Index)
 		{
 			FStaticMeshComponentLODInfo& InstanceMeshLODInfo = StaticMeshComponent->LODData[Index];
