Better Fix : SetLODDataCount to StaticMesh LOD Count

diff --git a/Engine/Source/Editor/MeshPaint/Private/MeshPaintStaticMeshAdapter.cpp b/Engine/Source/Editor/MeshPaint/Private/MeshPaintStaticMeshAdapter.cpp
--- a/Engine/Source/Editor/MeshPaint/Private/MeshPaintStaticMeshAdapter.cpp
+++ b/Engine/Source/Editor/MeshPaint/Private/MeshPaintStaticMeshAdapter.cpp
@@ -376,14 +376,15 @@ void FMeshPaintGeometryAdapterForStaticMeshes::PreEdit()
 		StaticMeshComponent->Modify();
 		StaticMeshComponent->bCustomOverrideVertexColorPerLOD = (MeshLODIndex > 0);
 				
-		const int32 MaxLOD = (MeshLODIndex == 0) ? StaticMesh->RenderData->LODResources.Num() : (MeshLODIndex + 1);
+		const int32 NumLODs = StaticMesh->GetNumLODs();
+		const int32 MaxIndex = (MeshLODIndex == 0) ? NumLODs : (MeshLODIndex + 1);
 		// Ensure LODData has enough entries in it, free not required.
-		StaticMeshComponent->SetLODDataCount(MeshLODIndex + 1, StaticMeshComponent->LODData.Num());
+
+		StaticMeshComponent->SetLODDataCount(NumLODs, NumLODs);
 
 		// If LOD is 0, pre-edit all LODs. There's currently no way to tell from here
 		// if VertexPaintSettings.bPaintOnSpecificLOD is set to true or not.
-		
-		for (int32 Index = MeshLODIndex; Index < MaxLOD; ++Index)
+		for (int32 Index = MeshLODIndex; Index < MaxIndex; ++Index)
 		{
 			FStaticMeshComponentLODInfo& InstanceMeshLODInfo = StaticMeshComponent->LODData[Index];
 			FStaticMeshLODResources& LODResource = StaticMesh->RenderData->LODResources[Index];
