LOD Paint warnings open without enabling LOD Paint
#fix Make sure we only consider per  lod painted vertices if the component wide flag has been set as well

diff --git a/Engine/Source/Editor/MeshPaint/Private/MeshPaintHelpers.cpp b/Engine/Source/Editor/MeshPaint/Private/MeshPaintHelpers.cpp
--- a/Engine/Source/Editor/MeshPaint/Private/MeshPaintHelpers.cpp
+++ b/Engine/Source/Editor/MeshPaint/Private/MeshPaintHelpers.cpp
@@ -1372,18 +1372,23 @@ bool MeshPaintHelpers::DoesMeshComponentContainPerLODColors(const UMeshComponent
 	{
 		bPerLODColors = StaticMeshComponent->bCustomOverrideVertexColorPerLOD;
 		
-		const int32 NumLODs = StaticMeshComponent->LODData.Num();
 		bool bInstancedLODColors = false;
-		for (int32 LODIndex = 1; LODIndex < NumLODs; ++LODIndex)
-		{
-			if (StaticMeshComponent->LODData[LODIndex].PaintedVertices.Num() > 0)
-			{
-				bInstancedLODColors = true;
-				break;
-			}
+		if (bPerLODColors)
+        {
+
+		  const int32 NumLODs = StaticMeshComponent->LODData.Num();
+		  
+		  for (int32 LODIndex = 1; LODIndex < NumLODs; ++LODIndex)
+		  {
+			  if (StaticMeshComponent->LODData[LODIndex].PaintedVertices.Num() > 0)
+			  {
+				  bInstancedLODColors = true;
+				  break;
+			  }
+		  }
 		}
 
-		bPerLODColors = bInstancedLODColors;		
+		bPerLODColors = bPerLODColors && bInstancedLODColors;
 	}
 	else if (const USkeletalMeshComponent* SkeletalMeshComponent = Cast<USkeletalMeshComponent>(MeshComponent))
 	{
