Fix propagating vertex color to custom LOD

diff --git a/Engine/Source/Developer/MeshBuilder/Private/StaticMeshBuilder.cpp b/Engine/Source/Developer/MeshBuilder/Private/StaticMeshBuilder.cpp
--- a/Engine/Source/Developer/MeshBuilder/Private/StaticMeshBuilder.cpp
+++ b/Engine/Source/Developer/MeshBuilder/Private/StaticMeshBuilder.cpp
@@ -261,9 +261,8 @@ bool FStaticMeshBuilder::Build(FStaticMeshRenderData& StaticMeshRenderData, USta
 
 		StaticMeshLOD.Sections.Empty(PolygonGroups.Num());
 		TArray<int32> RemapVerts; //Because we will remove MeshVertex that are redundant, we need a remap
-								  //Render data Wedge map is only set for LOD 0???
-		TArray<int32> TempWedgeMap;
-		TArray<int32> &WedgeMap = (LodIndex == 0) ? StaticMeshRenderData.WedgeMap : TempWedgeMap;
+		TArray<int32> &WedgeMap = StaticMeshLOD.WedgeMap;
+		WedgeMap.Reset();
 
 		//Prepare the PerSectionIndices array so we can optimize the index buffer for the GPU
 		TArray<TArray<uint32> > PerSectionIndices;
diff --git a/Engine/Source/Developer/MeshUtilities/Private/MeshUtilities.cpp b/Engine/Source/Developer/MeshUtilities/Private/MeshUtilities.cpp
--- a/Engine/Source/Developer/MeshUtilities/Private/MeshUtilities.cpp
+++ b/Engine/Source/Developer/MeshUtilities/Private/MeshUtilities.cpp
@@ -2772,7 +2772,8 @@ public:
 				// TODO_STATICMESH: The wedge map is only valid for LODIndex 0 if no reduction has been performed.
 				// We can compute an approximate one instead for other LODs.
 				TArray<int32> TempWedgeMap;
-				TArray<int32>& WedgeMap = (LODIndex == 0 && InOutModels[0].ReductionSettings.PercentTriangles >= 1.0f) ? OutRenderData.WedgeMap : TempWedgeMap;
+				TArray<int32>& WedgeMap = InOutModels[LODIndex].ReductionSettings.PercentTriangles >= 1.0f ? LODModel.WedgeMap : TempWedgeMap;
+				WedgeMap.Reset();
 				float ComparisonThreshold = GetComparisonThreshold(LODBuildSettings[LODIndex]);
 				MeshUtilities.BuildStaticMeshVertexAndIndexBuffers(Vertices, PerSectionIndices, WedgeMap, RawMesh, LODOverlappingCorners[LODIndex], MaterialToSectionMapping, ComparisonThreshold, LODBuildSettings[LODIndex].BuildScale3D, ImportVersion);
 				check(WedgeMap.Num() == RawMesh.WedgeIndices.Num());
diff --git a/Engine/Source/Editor/MeshPaint/Private/MeshPaintHelpers.cpp b/Engine/Source/Editor/MeshPaint/Private/MeshPaintHelpers.cpp
--- a/Engine/Source/Editor/MeshPaint/Private/MeshPaintHelpers.cpp
+++ b/Engine/Source/Editor/MeshPaint/Private/MeshPaintHelpers.cpp
@@ -106,7 +106,7 @@ bool MeshPaintHelpers::PropagateColorsToRawMesh(UStaticMesh* StaticMesh, int32 L
 	FStaticMeshRenderData& RenderData = *StaticMesh->RenderData;
 	FStaticMeshLODResources& RenderModel = RenderData.LODResources[LODIndex];
 	FColorVertexBuffer& ColorVertexBuffer = *ComponentLODInfo.OverrideVertexColors;
-	if (RenderData.WedgeMap.Num() > 0 && ColorVertexBuffer.GetNumVertices() == RenderModel.GetNumVertices())
+	if (RenderModel.WedgeMap.Num() > 0 && ColorVertexBuffer.GetNumVertices() == RenderModel.GetNumVertices())
 	{
 		// Use the wedge map if it is available as it is lossless.
 		FMeshDescription* MeshDescription = StaticMesh->GetMeshDescription(LODIndex);
@@ -119,14 +119,14 @@ bool MeshPaintHelpers::PropagateColorsToRawMesh(UStaticMesh* StaticMesh, int32 L
 
 		FStaticMeshAttributes Attributes(*MeshDescription);
 		int32 NumWedges = MeshDescription->VertexInstances().Num();
-		if (RenderData.WedgeMap.Num() == NumWedges)
+		if (RenderModel.WedgeMap.Num() == NumWedges)
 		{
 			TVertexInstanceAttributesRef<FVector4> Colors = Attributes.GetVertexInstanceColors();
 			int32 VertexInstanceIndex = 0;
 			for (const FVertexInstanceID VertexInstanceID : MeshDescription->VertexInstances().GetElementIDs())
 			{
 				FLinearColor WedgeColor = FLinearColor::White;
-				int32 Index = RenderData.WedgeMap[VertexInstanceIndex];
+				int32 Index = RenderModel.WedgeMap[VertexInstanceIndex];
 				if (Index != INDEX_NONE)
 				{
 					WedgeColor = FLinearColor(ColorVertexBuffer.VertexColor(Index));
@@ -160,7 +160,7 @@ bool MeshPaintHelpers::PropagateColorsToRawMesh(UStaticMesh* StaticMesh, int32 L
 		int32 VertexIndex = 0;
 		for (const FVertexID VertexID : MeshDescription->Vertices().GetElementIDs())
 		{
-			VertexPositionsDup[VertexIndex] = VertexPositions[VertexID];
+			VertexPositionsDup[VertexIndex++] = VertexPositions[VertexID];
 		}
 		TempPositionVertexBuffer.Init(VertexPositionsDup);
 		RemapPaintedVertexColors(
diff --git a/Engine/Source/Runtime/Engine/Private/StaticMesh.cpp b/Engine/Source/Runtime/Engine/Private/StaticMesh.cpp
--- a/Engine/Source/Runtime/Engine/Private/StaticMesh.cpp
+++ b/Engine/Source/Runtime/Engine/Private/StaticMesh.cpp
@@ -614,6 +614,13 @@ void FStaticMeshLODResources::Serialize(FArchive& Ar, UObject* Owner, int32 Inde
 	Ar << Sections;
 	Ar << MaxDeviation;
 
+#if WITH_EDITORONLY_DATA
+	if ((!Ar.IsCooking() && !Ar.IsFilterEditorOnly()) || (Ar.IsCooking() && Ar.CookingTarget()->HasEditorOnlyData()))
+	{
+		Ar << WedgeMap;
+	}
+#endif // #if WITH_EDITORONLY_DATA
+
 	const bool bIsBelowMinLOD = StripFlags.IsClassDataStripped(CDSF_MinLodData);
 	bool bIsLODCookedOut = IsLODCookedOut(Ar.CookingTarget(), OwnerStaticMesh, bIsBelowMinLOD);
 	Ar << bIsLODCookedOut;
@@ -1328,7 +1335,6 @@ void FStaticMeshRenderData::Serialize(FArchive& Ar, UStaticMesh* Owner, bool bCo
 #if WITH_EDITORONLY_DATA
 	if (!bCooked)
 	{
-		Ar << WedgeMap;
 		Ar << MaterialIndexToImportIndex;
 	}
 
diff --git a/Engine/Source/Runtime/Engine/Public/StaticMeshResources.h b/Engine/Source/Runtime/Engine/Public/StaticMeshResources.h
--- a/Engine/Source/Runtime/Engine/Public/StaticMeshResources.h
+++ b/Engine/Source/Runtime/Engine/Public/StaticMeshResources.h
@@ -391,6 +391,9 @@ struct FStaticMeshLODResources
 	FByteBulkData BulkData;
 
 	FString DerivedDataKey;
+
+	/** Map of wedge index to vertex index. Each LOD need one*/
+	TArray<int32> WedgeMap;
 #endif
 	
 	/** Default constructor. */
@@ -581,9 +584,6 @@ public:
 	/** The derived data key associated with this render data. */
 	FString DerivedDataKey;
 
-	/** Map of wedge index to vertex index. */
-	TArray<int32> WedgeMap;
-
 	/** Map of material index -> original material index at import time. */
 	TArray<int32> MaterialIndexToImportIndex;
 
