[NO NON-PII COMMIT MESSAGE]

diff --git a/Engine/Source/Editor/MeshPaint/Private/IMeshPainter.cpp b/Engine/Source/Editor/MeshPaint/Private/IMeshPainter.cpp
--- a/Engine/Source/Editor/MeshPaint/Private/IMeshPainter.cpp
+++ b/Engine/Source/Editor/MeshPaint/Private/IMeshPainter.cpp
@@ -21,14 +21,22 @@
 #include "Components/MeshComponent.h"
 #include "Components/PrimitiveComponent.h"
 
-IMeshPainter::IMeshPainter()
-	: CurrentViewportInteractor(nullptr), bArePainting(false), TimeSinceStartedPainting(0.0f), Time(0.0f), WidgetLineThickness(1.0f), VertexPointColor(FLinearColor::White), HoverVertexPointColor(0.3f, 1.0f, 0.3f), PaintTransaction(nullptr)
+IMeshPainter::IMeshPainter() :
+	CurrentViewportInteractor(nullptr),
+	bArePainting(false),
+	TimeSinceStartedPainting(0.0f),
+	Time(0.0f),
+	WidgetLineThickness(1.0f),
+	VertexPointColor(FLinearColor::White),
+	HoverVertexPointColor(0.3f, 1.0f, 0.3f),
+	PaintTransaction(nullptr)
 {
 	FMeshPainterCommands::Register();
 }
 
 IMeshPainter::~IMeshPainter()
 {
+	UICommandList.Reset();
 }
 
 void IMeshPainter::RenderInteractors(const FSceneView* View, FViewport* Viewport, FPrimitiveDrawInterface* PDI, bool bRenderVertices, ESceneDepthPriorityGroup DepthGroup/* = SDPG_World*/)
@@ -66,24 +74,53 @@ void IMeshPainter::Tick(FEditorViewportClient* ViewportClient, float DeltaTime)
 	Time += DeltaTime;
 }
 
+void IMeshPainter::ChangeBrushRadius(float Multiplier)
+{
+	const float ChangePercentage = 0.05f;
+	const float OldValue = GetBrushSettings()->GetBrushRadius();
+
+	float NewValue = OldValue * (1 + ChangePercentage * Multiplier);
+	GetBrushSettings()->SetBrushRadius(NewValue);
+}
+
+void IMeshPainter::ChangeBrushStrength(float Multiplier)
+{
+	const float ChangeAmount = 0.02f;
+	const float OldValue = GetBrushSettings()->GetBrushStrength();
+
+	float NewValue = OldValue + ChangeAmount * Multiplier;
+	GetBrushSettings()->SetBrushStrength(NewValue);
+}
+
+void IMeshPainter::ChangeBrushFalloff(float Multiplier)
+{
+	const float ChangeAmount = 0.02f;
+	const float OldValue = GetBrushSettings()->GetBrushFalloff();
+
+	float NewValue = OldValue + ChangeAmount * Multiplier;
+	GetBrushSettings()->SetBrushFalloff(NewValue);
+}
+
 void IMeshPainter::RegisterCommands(TSharedRef<FUICommandList> CommandList)
 {
 	const FMeshPainterCommands& Commands = FMeshPainterCommands::Get();
 
-	// Lambda used to alter the paint brush size
-	auto BrushLambda = [this](float Multiplier)
-	{
-		const float BrushChangeValue = 0.05f;
-		float BrushRadius = GetBrushSettings()->GetBrushRadius();
-		BrushRadius *= (1.f + (BrushChangeValue * Multiplier));
-		GetBrushSettings()->SetBrushRadius(BrushRadius);
-	};
-	CommandList->MapAction(Commands.IncreaseBrushSize, FExecuteAction::CreateLambda(BrushLambda, 1.0f), FCanExecuteAction(), EUIActionRepeatMode::RepeatEnabled);
-	CommandList->MapAction(Commands.DecreaseBrushSize, FExecuteAction::CreateLambda(BrushLambda, -1.0f), FCanExecuteAction(), EUIActionRepeatMode::RepeatEnabled);
+	CommandList->MapAction(Commands.IncreaseBrushRadius, FExecuteAction::CreateRaw(this, &IMeshPainter::ChangeBrushRadius, 1.0f), FCanExecuteAction(), EUIActionRepeatMode::RepeatEnabled);
+	CommandList->MapAction(Commands.DecreaseBrushRadius, FExecuteAction::CreateRaw(this, &IMeshPainter::ChangeBrushRadius, -1.0f), FCanExecuteAction(), EUIActionRepeatMode::RepeatEnabled);
+
+	CommandList->MapAction(Commands.IncreaseBrushStrength, FExecuteAction::CreateRaw(this, &IMeshPainter::ChangeBrushStrength, 1.0f), FCanExecuteAction(), EUIActionRepeatMode::RepeatEnabled);
+	CommandList->MapAction(Commands.DecreaseBrushStrength, FExecuteAction::CreateRaw(this, &IMeshPainter::ChangeBrushStrength, -1.0f), FCanExecuteAction(), EUIActionRepeatMode::RepeatEnabled);
+
+	CommandList->MapAction(Commands.IncreaseBrushFalloff, FExecuteAction::CreateRaw(this, &IMeshPainter::ChangeBrushFalloff, 1.0f), FCanExecuteAction(), EUIActionRepeatMode::RepeatEnabled);
+	CommandList->MapAction(Commands.DecreaseBrushFalloff, FExecuteAction::CreateRaw(this, &IMeshPainter::ChangeBrushFalloff, -1.0f), FCanExecuteAction(), EUIActionRepeatMode::RepeatEnabled);
+
+	UICommandList = CommandList;
 }
 
 void IMeshPainter::UnregisterCommands(TSharedRef<FUICommandList> CommandList)
 {
+	UICommandList.Reset();
+
 	const FMeshPainterCommands& Commands = FMeshPainterCommands::Get();	
 	for (const TSharedPtr<const FUICommandInfo> Action : Commands.Commands)
 	{
@@ -142,26 +179,11 @@ bool IMeshPainter::InputKey(FEditorViewportClient* InViewportClient, FViewport*
 {
 	bool bHandled = false;
 
-	// Lambda used to alter the paint brush size
-	auto BrushLambda = [this](float Multiplier)
+	if (UICommandList.IsValid())
 	{
-		const float BrushChangeValue = 0.05f;
-		float BrushRadius = GetBrushSettings()->GetBrushRadius();
-		BrushRadius *= (1.f + (BrushChangeValue * Multiplier));
-		GetBrushSettings()->SetBrushRadius(BrushRadius);
-	};
-
-	if (InKey == EKeys::LeftBracket)
-	{
-		BrushLambda(-1.0f);
-		bHandled = true;
-	}
-	else if (InKey == EKeys::RightBracket)
-	{
-		BrushLambda(1.0f);
-		bHandled = true;
+		bHandled = UICommandList->ProcessCommandBindings(InKey, FSlateApplication::Get().GetModifierKeys(), InEvent == IE_Repeat);
 	}
-
+	
 	return bHandled;
 }
 
diff --git a/Engine/Source/Editor/MeshPaint/Private/MeshPaintSettings.cpp b/Engine/Source/Editor/MeshPaint/Private/MeshPaintSettings.cpp
--- a/Engine/Source/Editor/MeshPaint/Private/MeshPaintSettings.cpp
+++ b/Engine/Source/Editor/MeshPaint/Private/MeshPaintSettings.cpp
@@ -16,7 +16,13 @@ UPaintBrushSettings::UPaintBrushSettings(const FObjectInitializer& ObjectInitial
 	BrushRadiusMin = 0.01f, BrushRadiusMax = 250000.0f;
 
 	GConfig->GetFloat(TEXT("MeshPaintEdit"), TEXT("DefaultBrushRadius"), BrushRadius, GEditorPerProjectIni);
-	BrushRadius = (float)FMath::Clamp(BrushRadius, BrushRadiusMin, BrushRadiusMax);
+	BrushRadius = FMath::Clamp(BrushRadius, BrushRadiusMin, BrushRadiusMax);
+
+	GConfig->GetFloat(TEXT("MeshPaintEdit"), TEXT("DefaultBrushStrength"), BrushStrength, GEditorPerProjectIni);
+	BrushStrength = FMath::Clamp(BrushRadius, 0.f, 1.f);
+
+	GConfig->GetFloat(TEXT("MeshPaintEdit"), TEXT("DefaultBrushFalloff"), BrushFalloffAmount, GEditorPerProjectIni);
+	BrushFalloffAmount = FMath::Clamp(BrushFalloffAmount, 0.f, 1.f);
 }
 
 UPaintBrushSettings::~UPaintBrushSettings()
@@ -29,10 +35,32 @@ void UPaintBrushSettings::SetBrushRadius(float InRadius)
 	GConfig->SetFloat(TEXT("MeshPaintEdit"), TEXT("DefaultBrushRadius"), BrushRadius, GEditorPerProjectIni);
 }
 
+void UPaintBrushSettings::SetBrushStrength(float InStrength)
+{
+	BrushStrength = FMath::Clamp(InStrength, 0.f, 1.f);
+	GConfig->SetFloat(TEXT("MeshPaintEdit"), TEXT("DefaultBrushStrength"), BrushStrength, GEditorPerProjectIni);
+}
+
+void UPaintBrushSettings::SetBrushFalloff(float InFalloff)
+{
+	BrushFalloffAmount = FMath::Clamp(InFalloff, 0.f, 1.f);
+	GConfig->SetFloat(TEXT("MeshPaintEdit"), TEXT("DefaultBrushFalloff"), BrushFalloffAmount, GEditorPerProjectIni);
+}
+
 void UPaintBrushSettings::PostEditChangeProperty(struct FPropertyChangedEvent& PropertyChangedEvent)
 {
 	if (PropertyChangedEvent.Property && PropertyChangedEvent.Property->GetFName() == GET_MEMBER_NAME_CHECKED(UPaintBrushSettings, BrushRadius) && PropertyChangedEvent.ChangeType != EPropertyChangeType::Interactive)
 	{
 		GConfig->SetFloat(TEXT("MeshPaintEdit"), TEXT("DefaultBrushRadius"), BrushRadius, GEditorPerProjectIni);
 	}
+
+	if (PropertyChangedEvent.Property && PropertyChangedEvent.Property->GetFName() == GET_MEMBER_NAME_CHECKED(UPaintBrushSettings, BrushStrength) && PropertyChangedEvent.ChangeType != EPropertyChangeType::Interactive)
+	{
+		GConfig->SetFloat(TEXT("MeshPaintEdit"), TEXT("DefaultBrushStrength"), BrushStrength, GEditorPerProjectIni);
+	}
+
+	if (PropertyChangedEvent.Property && PropertyChangedEvent.Property->GetFName() == GET_MEMBER_NAME_CHECKED(UPaintBrushSettings, BrushFalloffAmount) && PropertyChangedEvent.ChangeType != EPropertyChangeType::Interactive)
+	{
+		GConfig->SetFloat(TEXT("MeshPaintEdit"), TEXT("DefaultBrushFalloff"), BrushFalloffAmount, GEditorPerProjectIni);
+	}
 }
\ No newline at end of file
diff --git a/Engine/Source/Editor/MeshPaint/Private/MeshPainterCommands.cpp b/Engine/Source/Editor/MeshPaint/Private/MeshPainterCommands.cpp
--- a/Engine/Source/Editor/MeshPaint/Private/MeshPainterCommands.cpp
+++ b/Engine/Source/Editor/MeshPaint/Private/MeshPainterCommands.cpp
@@ -4,13 +4,31 @@
 
 #define LOCTEXT_NAMESPACE "MeshPainterCommands"
 
+FMeshPainterCommands::FMeshPainterCommands() 
+	: TCommands<FMeshPainterCommands>(
+		"MeshPainter", 
+		NSLOCTEXT("Contexts", "MeshPainter", "Mesh Painter"), 
+		NAME_None, 
+		FEditorStyle::GetStyleSetName())
+{
+}
+
 void FMeshPainterCommands::RegisterCommands()
 {
-	UI_COMMAND(IncreaseBrushSize, "IncreaseBrush", "Increases brush size", EUserInterfaceActionType::Button, FInputChord(EKeys::RightBracket, EModifierKey::Control));
-	Commands.Add(IncreaseBrushSize);
+	UI_COMMAND(IncreaseBrushRadius, "Increase Brush Radius", "Press this key to increase brush radius by a percentage of its current size.", EUserInterfaceActionType::Button, FInputChord(EKeys::RightBracket));
+	Commands.Add(IncreaseBrushRadius);
+	UI_COMMAND(DecreaseBrushRadius, "Decrease Brush Size", "Press this key to decrease brush radius by a percentage of its current size.", EUserInterfaceActionType::Button, FInputChord(EKeys::LeftBracket));
+	Commands.Add(DecreaseBrushRadius);
+
+	UI_COMMAND(IncreaseBrushStrength, "Increase Brush Strength", "Press this key to increase brush strength by a fixed increment.", EUserInterfaceActionType::Button, FInputChord(EModifierKey::Control, EKeys::RightBracket));
+	Commands.Add(IncreaseBrushStrength);
+	UI_COMMAND(DecreaseBrushStrength, "Decrease Brush Strength", "Press this key to decrease brush strength by a fixed increment.", EUserInterfaceActionType::Button, FInputChord(EModifierKey::Control, EKeys::LeftBracket));
+	Commands.Add(DecreaseBrushStrength);
 
-	UI_COMMAND(DecreaseBrushSize, "DecreaseBrush", "Decreases brush size", EUserInterfaceActionType::Button, FInputChord(EKeys::LeftBracket, EModifierKey::Control));
-	Commands.Add(DecreaseBrushSize);
+	UI_COMMAND(IncreaseBrushFalloff, "Increase Brush Falloff", "Press this key to increase brush falloff by a fixed increment.", EUserInterfaceActionType::Button, FInputChord(FInputChord(EModifierKey::Control | EModifierKey::Shift, EKeys::RightBracket)));
+	Commands.Add(IncreaseBrushFalloff);
+	UI_COMMAND(DecreaseBrushFalloff, "Decrease Brush Falloff", "Press this key to decrease brush falloff by a fixed increment.", EUserInterfaceActionType::Button, FInputChord(FInputChord(EModifierKey::Control | EModifierKey::Shift, EKeys::LeftBracket)));
+	Commands.Add(DecreaseBrushFalloff);
 }
 
 #undef LOCTEXT_NAMESPACE
diff --git a/Engine/Source/Editor/MeshPaint/Public/IMeshPainter.h b/Engine/Source/Editor/MeshPaint/Public/IMeshPainter.h
--- a/Engine/Source/Editor/MeshPaint/Public/IMeshPainter.h
+++ b/Engine/Source/Editor/MeshPaint/Public/IMeshPainter.h
@@ -112,6 +112,9 @@ protected:
 	FLinearColor WidgetLineColor;
 	FLinearColor VertexPointColor;
 	FLinearColor HoverVertexPointColor;
+
+	TSharedPtr<FUICommandList> UICommandList;
+
 protected:
 	/** Begins and keeps track of an Editor transaction with the given Description */
 	void BeginTransaction(const FText Description);
@@ -119,4 +122,15 @@ protected:
 	void EndTransaction();
 	/** Painting transaction instance which is currently active */
 	FScopedTransaction* PaintTransaction;
+
+private:
+
+	/** Change the brush radius by a percentage, multiplied by the given value. */
+	void ChangeBrushRadius(float Multiplier);
+
+	/** Change the brush radius by one increment, multiplied by the given value */
+	void ChangeBrushStrength(float Multiplier);
+
+	/** Change the brush radius by one increment, multiplied by the given value */
+	void ChangeBrushFalloff(float Multiplier);
 };
\ No newline at end of file
diff --git a/Engine/Source/Editor/MeshPaint/Public/MeshPaintSettings.h b/Engine/Source/Editor/MeshPaint/Public/MeshPaintSettings.h
--- a/Engine/Source/Editor/MeshPaint/Public/MeshPaintSettings.h
+++ b/Engine/Source/Editor/MeshPaint/Public/MeshPaintSettings.h
@@ -41,6 +41,13 @@ public:
 
 	float GetBrushRadius() const { return BrushRadius; }
 	void SetBrushRadius(float InRadius);
+
+	float GetBrushStrength() const { return BrushStrength; }
+	void SetBrushStrength(float InStrength);
+
+	float GetBrushFalloff() const { return BrushFalloffAmount; }
+	void SetBrushFalloff(float InFalloff);
+
 protected:
 	/** Radius of the Brush used for Painting */
 	UPROPERTY(EditAnywhere, Category = Brush, meta = (DisplayName = "Radius", UIMin = "0.01", UIMax = "2048.0", ClampMin = "0.01", ClampMax = "250000.0"))
diff --git a/Engine/Source/Editor/MeshPaint/Public/MeshPainterCommands.h b/Engine/Source/Editor/MeshPaint/Public/MeshPainterCommands.h
--- a/Engine/Source/Editor/MeshPaint/Public/MeshPainterCommands.h
+++ b/Engine/Source/Editor/MeshPaint/Public/MeshPainterCommands.h
@@ -9,9 +9,8 @@
 /** Base set of mesh painter commands */
 class MESHPAINT_API FMeshPainterCommands : public TCommands<FMeshPainterCommands>
 {
-
 public:
-	FMeshPainterCommands() : TCommands<FMeshPainterCommands>("MeshPainter", NSLOCTEXT("Contexts", "MeshPainter", "Mesh Painter"), NAME_None, FEditorStyle::GetStyleSetName()) {}
+	FMeshPainterCommands();
 
 	/**
 	* Initialize commands
@@ -19,7 +18,14 @@ public:
 	virtual void RegisterCommands() override;
 
 public:
-	TSharedPtr<FUICommandInfo> IncreaseBrushSize;
-	TSharedPtr<FUICommandInfo> DecreaseBrushSize;
+	TSharedPtr<FUICommandInfo> IncreaseBrushRadius;
+	TSharedPtr<FUICommandInfo> DecreaseBrushRadius;
+
+	TSharedPtr<FUICommandInfo> IncreaseBrushStrength;
+	TSharedPtr<FUICommandInfo> DecreaseBrushStrength;
+
+	TSharedPtr<FUICommandInfo> IncreaseBrushFalloff;
+	TSharedPtr<FUICommandInfo> DecreaseBrushFalloff;
+
 	TArray<TSharedPtr<FUICommandInfo>> Commands;
 };
\ No newline at end of file
