[NO NON-PII COMMIT MESSAGE]

diff --git a/Engine/Source/Editor/MeshPaint/Private/IMeshPainter.cpp b/Engine/Source/Editor/MeshPaint/Private/IMeshPainter.cpp
--- a/Engine/Source/Editor/MeshPaint/Private/IMeshPainter.cpp
+++ b/Engine/Source/Editor/MeshPaint/Private/IMeshPainter.cpp
@@ -17,9 +17,10 @@
 #include "VREditorInteractor.h"
 #include "EditorWorldExtension.h"
 
-#include "Framework/Commands/UICommandList.h"
 #include "Components/MeshComponent.h"
 #include "Components/PrimitiveComponent.h"
+#include "Framework/Application/SlateApplication.h"
+#include "Framework/Commands/UICommandList.h"
 
 IMeshPainter::IMeshPainter() :
 	CurrentViewportInteractor(nullptr),
