Only load the MeshPaintCommands in the MeshPaint modules, not the MeshPaintMode module as it shares a weak symbol to the same instance type. Since we load with RTLD_LOCAL it will prefer  to use local versions of symbols over the first matched one

diff --git a/Engine/Source/Editor/MeshPaint/Private/MeshPaintModule.cpp b/Engine/Source/Editor/MeshPaint/Private/MeshPaintModule.cpp
--- a/Engine/Source/Editor/MeshPaint/Private/MeshPaintModule.cpp
+++ b/Engine/Source/Editor/MeshPaint/Private/MeshPaintModule.cpp
@@ -6,6 +6,7 @@
 #include "EditorStyleSet.h"
 #include "EditorModeRegistry.h"
 #include "EditorModes.h"
+#include "MeshPainterCommands.h"
 #include "MeshPaintAdapterFactory.h"
 #include "MeshPaintStaticMeshAdapter.h"
 #include "MeshPaintSplineMeshAdapter.h"
@@ -38,6 +39,8 @@ public:
 	 */
 	virtual void StartupModule() override
 	{
+		FMeshPainterCommands::Register();
+
 		RegisterGeometryAdapterFactory(MakeShareable(new FMeshPaintGeometryAdapterForSplineMeshesFactory));
 		RegisterGeometryAdapterFactory(MakeShareable(new FMeshPaintGeometryAdapterForStaticMeshesFactory));
 		RegisterGeometryAdapterFactory(MakeShareable(new FMeshPaintGeometryAdapterForSkeletalMeshesFactory));
@@ -60,6 +63,8 @@ public:
 
 	virtual void ShutdownModule() override
 	{
+		FMeshPainterCommands::Unregister();
+
 		FPropertyEditorModule& PropertyModule = FModuleManager::LoadModuleChecked<FPropertyEditorModule>("PropertyEditor");
 		PropertyModule.UnregisterCustomClassLayout("VertexColorImportOptions");
 
diff --git a/Engine/Source/Editor/MeshPaintMode/Private/MeshPaintModeModule.cpp b/Engine/Source/Editor/MeshPaintMode/Private/MeshPaintModeModule.cpp
--- a/Engine/Source/Editor/MeshPaintMode/Private/MeshPaintModeModule.cpp
+++ b/Engine/Source/Editor/MeshPaintMode/Private/MeshPaintModeModule.cpp
@@ -26,8 +26,6 @@ void FMeshPaintModeModule::StartupModule()
 		FSlateIcon(FEditorStyle::GetStyleSetName(), "LevelEditor.MeshPaintMode", "LevelEditor.MeshPaintMode.Small"),
 		true, 200 );
 
-	FMeshPainterCommands::Register();
-
 	/** Register detail/property customization */
 	FPropertyEditorModule& PropertyModule = FModuleManager::LoadModuleChecked<FPropertyEditorModule>("PropertyEditor");
 	PropertyModule.RegisterCustomClassLayout("PaintModeSettings", FOnGetDetailCustomizationInstance::CreateStatic(&FPaintModeSettingsCustomization::MakeInstance));
@@ -40,8 +38,6 @@ void FMeshPaintModeModule::StartupModule()
 
 void FMeshPaintModeModule::ShutdownModule()
 {
-	FMeshPainterCommands::Unregister();
-
 	FEditorModeRegistry::Get().UnregisterMode(FBuiltinEditorModes::EM_MeshPaint);
 
 	/** De-register detail/property customization */
