GetFLOATMetadata and GetINTMetadata renamed to CamelCase to be less shouty.

diff --git a/Engine/Plugins/Experimental/AlembicImporter/Source/AlembicLibrary/Private/AbcImportSettingsCustomization.cpp b/Engine/Plugins/Experimental/AlembicImporter/Source/AlembicLibrary/Private/AbcImportSettingsCustomization.cpp
--- a/Engine/Plugins/Experimental/AlembicImporter/Source/AlembicLibrary/Private/AbcImportSettingsCustomization.cpp
+++ b/Engine/Plugins/Experimental/AlembicImporter/Source/AlembicLibrary/Private/AbcImportSettingsCustomization.cpp
@@ -84,7 +84,7 @@ void FAbcSamplingSettingsCustomization::CustomizeChildren(TSharedRef<IPropertyHa
 
 		IDetailPropertyRow& Property = StructBuilder.AddProperty(ChildHandle);
 		static const FName EditConditionName = "EnumCondition";
-		int32 EnumCondition = ChildHandle->GetINTMetaData(EditConditionName);
+		int32 EnumCondition = ChildHandle->GetIntMetaData(EditConditionName);
 		Property.Visibility(TAttribute<EVisibility>::Create(TAttribute<EVisibility>::FGetter::CreateSP(this, &FAbcSamplingSettingsCustomization::ArePropertiesVisible, EnumCondition)));
 	}
 }
@@ -121,7 +121,7 @@ void FAbcCompressionSettingsCustomization::CustomizeChildren(TSharedRef<IPropert
 
 		IDetailPropertyRow& Property = StructBuilder.AddProperty(ChildHandle);
 		static const FName EditConditionName = "EnumCondition";
-		int32 EnumCondition = ChildHandle->GetINTMetaData(EditConditionName);
+		int32 EnumCondition = ChildHandle->GetIntMetaData(EditConditionName);
 		Property.Visibility(TAttribute<EVisibility>::Create(TAttribute<EVisibility>::FGetter::CreateSP(this, &FAbcCompressionSettingsCustomization::ArePropertiesVisible, EnumCondition)));
 	}
 }
diff --git a/Engine/Plugins/Experimental/ControlRig/Source/ControlRigDeveloper/Private/ControlRigModel.cpp b/Engine/Plugins/Experimental/ControlRig/Source/ControlRigDeveloper/Private/ControlRigModel.cpp
--- a/Engine/Plugins/Experimental/ControlRig/Source/ControlRigDeveloper/Private/ControlRigModel.cpp
+++ b/Engine/Plugins/Experimental/ControlRig/Source/ControlRigDeveloper/Private/ControlRigModel.cpp
@@ -584,7 +584,7 @@ bool UControlRigModel::AddNode(const FControlRigModelNode& InNode, bool bUndo)
 				FString PinPath = AddedNode.GetPinPath(Pin.Index, false);
 				if (UArrayProperty* Property = Cast<UArrayProperty>(Struct->FindPropertyByName(*PinPath)))
 				{
-					int32 DefaultArraySize = Property->GetINTMetaData(UControlRig::DefaultArraySizeMetaName);
+					int32 DefaultArraySize = Property->GetIntMetaData(UControlRig::DefaultArraySizeMetaName);
 					if (DefaultArraySize > 0)
 					{
 						FPinArrayInfo Info;
diff --git a/Engine/Source/Editor/DetailCustomizations/Private/IntervalStructCustomization.cpp b/Engine/Source/Editor/DetailCustomizations/Private/IntervalStructCustomization.cpp
--- a/Engine/Source/Editor/DetailCustomizations/Private/IntervalStructCustomization.cpp
+++ b/Engine/Source/Editor/DetailCustomizations/Private/IntervalStructCustomization.cpp
@@ -26,7 +26,7 @@ namespace IntervalMetadata
 	{
 		static float GetMetaData(const UProperty* Property, const TCHAR* Key)
 		{
-			return Property->GetFLOATMetaData(Key);
+			return Property->GetFloatMetaData(Key);
 		}
 	};
 
@@ -35,7 +35,7 @@ namespace IntervalMetadata
 	{
 		static int32 GetMetaData(const UProperty* Property, const TCHAR* Key)
 		{
-			return Property->GetINTMetaData(Key);
+			return Property->GetIntMetaData(Key);
 		}
 	};
 }
diff --git a/Engine/Source/Editor/DetailCustomizations/Private/MeshMergingSettingsCustomization.cpp b/Engine/Source/Editor/DetailCustomizations/Private/MeshMergingSettingsCustomization.cpp
--- a/Engine/Source/Editor/DetailCustomizations/Private/MeshMergingSettingsCustomization.cpp
+++ b/Engine/Source/Editor/DetailCustomizations/Private/MeshMergingSettingsCustomization.cpp
@@ -38,7 +38,7 @@ void FMeshMergingSettingsObjectCustomization::CustomizeDetails(IDetailLayoutBuil
 		if (Property->GetProperty()->GetFName() == GET_MEMBER_NAME_CHECKED(FMeshMergingSettings, SpecificLOD))
 		{
 			static const FName EditConditionName = "EnumCondition";
-			int32 EnumCondition = Property->GetINTMetaData(EditConditionName);
+			int32 EnumCondition = Property->GetIntMetaData(EditConditionName);
 			PropertyRow.Visibility(TAttribute<EVisibility>::Create(TAttribute<EVisibility>::FGetter::CreateSP(this, &FMeshMergingSettingsObjectCustomization::ArePropertiesVisible, EnumCondition)));
 		}
 		else if (Property->GetProperty()->GetFName() == GET_MEMBER_NAME_CHECKED(FMeshMergingSettings, LODSelectionType))
diff --git a/Engine/Source/Editor/DetailCustomizations/Private/RangeStructCustomization.cpp b/Engine/Source/Editor/DetailCustomizations/Private/RangeStructCustomization.cpp
--- a/Engine/Source/Editor/DetailCustomizations/Private/RangeStructCustomization.cpp
+++ b/Engine/Source/Editor/DetailCustomizations/Private/RangeStructCustomization.cpp
@@ -32,7 +32,7 @@ namespace
 	{
 		static float GetMetaData(const UProperty* Property, const TCHAR* Key)
 		{
-			return Property->GetFLOATMetaData(Key);
+			return Property->GetFloatMetaData(Key);
 		}
 	};
 
@@ -41,7 +41,7 @@ namespace
 	{
 		static int32 GetMetaData(const UProperty* Property, const TCHAR* Key)
 		{
-			return Property->GetINTMetaData(Key);
+			return Property->GetIntMetaData(Key);
 		}
 	};
 }
diff --git a/Engine/Source/Editor/MeshPaint/Private/MeshPaintSettings.cpp b/Engine/Source/Editor/MeshPaint/Private/MeshPaintSettings.cpp
--- a/Engine/Source/Editor/MeshPaint/Private/MeshPaintSettings.cpp
+++ b/Engine/Source/Editor/MeshPaint/Private/MeshPaintSettings.cpp
@@ -18,8 +18,8 @@ UPaintBrushSettings::UPaintBrushSettings(const FObjectInitializer& ObjectInitial
 
 	{
 		const UProperty* BrushRadiusProperty = UPaintBrushSettings::StaticClass()->FindPropertyByName(GET_MEMBER_NAME_CHECKED(UPaintBrushSettings, BrushRadius));
-		BrushRadiusMin = BrushRadiusProperty->GetFLOATMetaData(ClampMin);
-		BrushRadiusMax = BrushRadiusProperty->GetFLOATMetaData(ClampMax);
+		BrushRadiusMin = BrushRadiusProperty->GetFloatMetaData(ClampMin);
+		BrushRadiusMax = BrushRadiusProperty->GetFloatMetaData(ClampMax);
 
 		GConfig->GetFloat(TEXT("MeshPaintEdit"), TEXT("DefaultBrushRadius"), BrushRadius, GEditorPerProjectIni);
 		BrushRadius = FMath::Clamp(BrushRadius, BrushRadiusMin, BrushRadiusMax);
@@ -27,8 +27,8 @@ UPaintBrushSettings::UPaintBrushSettings(const FObjectInitializer& ObjectInitial
 
 	{
 		const UProperty* BrushStrengthProperty = UPaintBrushSettings::StaticClass()->FindPropertyByName(GET_MEMBER_NAME_CHECKED(UPaintBrushSettings, BrushStrength));
-		float BrushStrengthMin = BrushStrengthProperty->GetFLOATMetaData(ClampMin);
-		float BrushStrengthMax = BrushStrengthProperty->GetFLOATMetaData(ClampMax);
+		float BrushStrengthMin = BrushStrengthProperty->GetFloatMetaData(ClampMin);
+		float BrushStrengthMax = BrushStrengthProperty->GetFloatMetaData(ClampMax);
 
 		GConfig->GetFloat(TEXT("MeshPaintEdit"), TEXT("DefaultBrushStrength"), BrushStrength, GEditorPerProjectIni);
 		BrushStrength = FMath::Clamp(BrushStrength, BrushStrengthMin, BrushStrengthMax);
@@ -36,8 +36,8 @@ UPaintBrushSettings::UPaintBrushSettings(const FObjectInitializer& ObjectInitial
 
 	{
 		const UProperty* BrushFalloffProperty = UPaintBrushSettings::StaticClass()->FindPropertyByName(GET_MEMBER_NAME_CHECKED(UPaintBrushSettings, BrushFalloffAmount));
-		float BrushFalloffMin = BrushFalloffProperty->GetFLOATMetaData(ClampMin);
-		float BrushFalloffMax = BrushFalloffProperty->GetFLOATMetaData(ClampMax);
+		float BrushFalloffMin = BrushFalloffProperty->GetFloatMetaData(ClampMin);
+		float BrushFalloffMax = BrushFalloffProperty->GetFloatMetaData(ClampMax);
 
 		GConfig->GetFloat(TEXT("MeshPaintEdit"), TEXT("DefaultBrushFalloff"), BrushFalloffAmount, GEditorPerProjectIni);
 		BrushFalloffAmount = FMath::Clamp(BrushFalloffAmount, BrushFalloffMin, BrushFalloffMax);
diff --git a/Engine/Source/Editor/MeshPaintMode/Private/PaintModeSettingsCustomization.cpp b/Engine/Source/Editor/MeshPaintMode/Private/PaintModeSettingsCustomization.cpp
--- a/Engine/Source/Editor/MeshPaintMode/Private/PaintModeSettingsCustomization.cpp
+++ b/Engine/Source/Editor/MeshPaintMode/Private/PaintModeSettingsCustomization.cpp
@@ -190,7 +190,7 @@ void FVertexPaintSettingsCustomization::CustomizeChildren(TSharedRef<IPropertyHa
 			static const FName EditConditionName = "EnumCondition";
 			if (ChildHandle->HasMetaData(EditConditionName))
 			{
-				int32 EnumCondition = ChildHandle->GetINTMetaData(EditConditionName);
+				int32 EnumCondition = ChildHandle->GetIntMetaData(EditConditionName);
 				Property.Visibility(TAttribute<EVisibility>::Create(TAttribute<EVisibility>::FGetter::CreateSP(this, &FVertexPaintSettingsCustomization::ArePropertiesVisible, EnumCondition)));
 			}
 		}
diff --git a/Engine/Source/Editor/PropertyEditor/Private/PropertyHandleImpl.cpp b/Engine/Source/Editor/PropertyEditor/Private/PropertyHandleImpl.cpp
--- a/Engine/Source/Editor/PropertyEditor/Private/PropertyHandleImpl.cpp
+++ b/Engine/Source/Editor/PropertyEditor/Private/PropertyHandleImpl.cpp
@@ -2581,16 +2581,16 @@ bool FPropertyHandleBase::GetBoolMetaData(const FName& Key) const
 	return (MetaDataProperty) ? MetaDataProperty->GetBoolMetaData(Key) : false;
 }
 
-int32 FPropertyHandleBase::GetINTMetaData(const FName& Key) const
+int32 FPropertyHandleBase::GetIntMetaData(const FName& Key) const
 {
 	UProperty* const MetaDataProperty = GetMetaDataProperty();
-	return (MetaDataProperty) ? MetaDataProperty->GetINTMetaData(Key) : 0;
+	return (MetaDataProperty) ? MetaDataProperty->GetIntMetaData(Key) : 0;
 }
 
-float FPropertyHandleBase::GetFLOATMetaData(const FName& Key) const
+float FPropertyHandleBase::GetFloatMetaData(const FName& Key) const
 {
 	UProperty* const MetaDataProperty = GetMetaDataProperty();
-	return (MetaDataProperty) ? MetaDataProperty->GetFLOATMetaData(Key) : 0.0f;
+	return (MetaDataProperty) ? MetaDataProperty->GetFloatMetaData(Key) : 0.0f;
 }
 
 UClass* FPropertyHandleBase::GetClassMetaData(const FName& Key) const
diff --git a/Engine/Source/Editor/PropertyEditor/Private/PropertyHandleImpl.h b/Engine/Source/Editor/PropertyEditor/Private/PropertyHandleImpl.h
--- a/Engine/Source/Editor/PropertyEditor/Private/PropertyHandleImpl.h
+++ b/Engine/Source/Editor/PropertyEditor/Private/PropertyHandleImpl.h
@@ -468,8 +468,8 @@ public:
 	virtual bool HasMetaData(const FName& Key) const override;
 	virtual const FString& GetMetaData(const FName& Key) const override;
 	virtual bool GetBoolMetaData(const FName& Key) const override;
-	virtual int32 GetINTMetaData(const FName& Key) const override;
-	virtual float GetFLOATMetaData(const FName& Key) const override;
+	virtual int32 GetIntMetaData(const FName& Key) const override;
+	virtual float GetFloatMetaData(const FName& Key) const override;
 	virtual UClass* GetClassMetaData(const FName& Key) const override;
 	virtual void SetInstanceMetaData(const FName& Key, const FString& Value) override;
 	virtual const FString* GetInstanceMetaData(const FName& Key) const override;
diff --git a/Engine/Source/Editor/PropertyEditor/Private/PropertyNode.cpp b/Engine/Source/Editor/PropertyEditor/Private/PropertyNode.cpp
--- a/Engine/Source/Editor/PropertyEditor/Private/PropertyNode.cpp
+++ b/Engine/Source/Editor/PropertyEditor/Private/PropertyNode.cpp
@@ -177,7 +177,7 @@ void FPropertyNode::InitNode( const FPropertyNodeInitParams& InitParams )
 		static const FName Name_MaxPropertyDepth("MaxPropertyDepth");
 		if (Property->HasMetaData(Name_MaxPropertyDepth))
 		{
-			int32 NewMaxChildDepthAllowed = Property->GetINTMetaData(Name_MaxPropertyDepth);
+			int32 NewMaxChildDepthAllowed = Property->GetIntMetaData(Name_MaxPropertyDepth);
 			//Ensure new depth is valid.  Otherwise just let the parent specified value stand
 			if (NewMaxChildDepthAllowed > 0)
 			{
diff --git a/Engine/Source/Editor/PropertyEditor/Public/PropertyHandle.h b/Engine/Source/Editor/PropertyEditor/Public/PropertyHandle.h
--- a/Engine/Source/Editor/PropertyEditor/Public/PropertyHandle.h
+++ b/Engine/Source/Editor/PropertyEditor/Public/PropertyHandle.h
@@ -95,7 +95,7 @@ public:
 	 * @param Key The key to lookup in the metadata
 	 * @return the int value stored in the metadata.
 	 */
-	virtual int32 GetINTMetaData(const FName& Key) const = 0;
+	virtual int32 GetIntMetaData(const FName& Key) const = 0;
 
 	/**
 	 * Find the metadata value associated with the key and return float
@@ -103,7 +103,7 @@ public:
 	 * @param Key The key to lookup in the metadata
 	 * @return the float value stored in the metadata.
 	 */
-	virtual float GetFLOATMetaData(const FName& Key) const = 0;
+	virtual float GetFloatMetaData(const FName& Key) const = 0;
 
 	/**
 	 * Find the metadata value associated with the key and return UClass*
diff --git a/Engine/Source/Editor/UnrealEd/Classes/Editor/PropertyEditorTestObject.h b/Engine/Source/Editor/UnrealEd/Classes/Editor/PropertyEditorTestObject.h
--- a/Engine/Source/Editor/UnrealEd/Classes/Editor/PropertyEditorTestObject.h
+++ b/Engine/Source/Editor/UnrealEd/Classes/Editor/PropertyEditorTestObject.h
@@ -520,4 +520,7 @@ class UPropertyEditorTestObject : public UObject
 
 	UPROPERTY(EditAnywhere, Category = EditConditionHides, meta = (EditConditionHides, EditCondition = "PropertyThatHides == EComponentMobility::Stationary"))
 	int32 VisibleWhenStationary;
+
+	UPROPERTY(EditAnywhere, Category = DateTime)
+	FDateTime DateTime;
 };
diff --git a/Engine/Source/Runtime/CoreUObject/Public/UObject/Class.h b/Engine/Source/Runtime/CoreUObject/Public/UObject/Class.h
--- a/Engine/Source/Runtime/CoreUObject/Public/UObject/Class.h
+++ b/Engine/Source/Runtime/CoreUObject/Public/UObject/Class.h
@@ -169,18 +169,18 @@ class COREUOBJECT_API UField : public UObject
 	 * Find the metadata value associated with the key
 	 * and return int32 
 	 * @param Key The key to lookup in the metadata
-	 * @return the int value stored in the metadata.
+	 * @return the int value stored in the metadata. 0 if not a valid integer.
 	 */
-	int32 GetINTMetaData(const TCHAR* Key) const
+	int32 GetIntMetaData(const TCHAR* Key) const
 	{
-		const FString& INTString = GetMetaData(Key);
-		int32 Value = FCString::Atoi(*INTString);
+		const FString& IntString = GetMetaData(Key);
+		int32 Value = FCString::Atoi(*IntString);
 		return Value;
 	}
-	int32 GetINTMetaData(const FName& Key) const
+	int32 GetIntMetaData(const FName& Key) const
 	{
-		const FString& INTString = GetMetaData(Key);
-		int32 Value = FCString::Atoi(*INTString);
+		const FString& IntString = GetMetaData(Key);
+		int32 Value = FCString::Atoi(*IntString);
 		return Value;
 	}
 
@@ -188,20 +188,18 @@ class COREUOBJECT_API UField : public UObject
 	 * Find the metadata value associated with the key
 	 * and return float
 	 * @param Key The key to lookup in the metadata
-	 * @return the float value stored in the metadata.
+	 * @return the float value stored in the metadata. 0 if not a valid float.
 	 */
-	float GetFLOATMetaData(const TCHAR* Key) const
+	float GetFloatMetaData(const TCHAR* Key) const
 	{
-		const FString& FLOATString = GetMetaData(Key);
-		// FString == operator does case insensitive comparison
-		float Value = FCString::Atof(*FLOATString);
+		const FString& FloatString = GetMetaData(Key);
+		float Value = FCString::Atof(*FloatString);
 		return Value;
 	}
-	float GetFLOATMetaData(const FName& Key) const
+	float GetFloatMetaData(const FName& Key) const
 	{
-		const FString& FLOATString = GetMetaData(Key);
-		// FString == operator does case insensitive comparison
-		float Value = FCString::Atof(*FLOATString);
+		const FString& FloatString = GetMetaData(Key);
+		float Value = FCString::Atof(*FloatString);
 		return Value;
 	}
 	
