Painting on HDR textures is unsupported and previously caused painting crashes. This will disallow the painting earlier in the process.

diff --git a/Engine/Source/Editor/MeshPaint/Private/IMeshPaintGeometryAdapter.cpp b/Engine/Source/Editor/MeshPaint/Private/IMeshPaintGeometryAdapter.cpp
--- a/Engine/Source/Editor/MeshPaint/Private/IMeshPaintGeometryAdapter.cpp
+++ b/Engine/Source/Editor/MeshPaint/Private/IMeshPaintGeometryAdapter.cpp
@@ -55,7 +55,8 @@ void IMeshPaintGeometryAdapter::DefaultQueryPaintableTextures(int32 MaterialInde
 			if (TextureBase != NULL &&
 				TextureBase->Texture != NULL &&
 				!TextureBase->Texture->IsNormalMap() && 
-				!TextureBase->Texture->VirtualTextureStreaming )
+				!TextureBase->Texture->VirtualTextureStreaming &&
+				!TextureBase->Texture->HasHDRSource()) // Currently HDR textures are not supported to paint on 
 			{
 				// Default UV channel to index 0. 
 				PaintableTexture = FPaintableTexture(TextureBase->Texture, 0);
