Unshelved from pending changelist '10954588':

Painting skeletalmesh is not save in the asset. Fallback of the sk refactor

diff --git a/Engine/Source/Editor/MeshPaint/Private/IMeshPaintMode.cpp b/Engine/Source/Editor/MeshPaint/Private/IMeshPaintMode.cpp
--- a/Engine/Source/Editor/MeshPaint/Private/IMeshPaintMode.cpp
+++ b/Engine/Source/Editor/MeshPaint/Private/IMeshPaintMode.cpp
@@ -215,12 +215,16 @@ bool IMeshPaintEdMode::ProcessCapturedMouseMoves( FEditorViewportClient* InViewp
 	return bPaintApplied;
 }
 
-/** FEdMode: Called when the a mouse button is released */
+/** FEdMode: Called when a mouse button is released */
 bool IMeshPaintEdMode::EndTracking(FEditorViewportClient* InViewportClient, FViewport* InViewport)
 {
 	InViewportClient->bLockFlightCamera = false;
-	MeshPainter->FinishPainting();
-	PaintingWithInteractorInVR = nullptr;
+	if (MeshPainter->IsPainting())
+	{
+		MeshPainter->FinishPainting();
+		PaintingWithInteractorInVR = nullptr;
+	}
+
 	return true;
 }
 
diff --git a/Engine/Source/Editor/MeshPaint/Private/MeshPaintHelpers.cpp b/Engine/Source/Editor/MeshPaint/Private/MeshPaintHelpers.cpp
--- a/Engine/Source/Editor/MeshPaint/Private/MeshPaintHelpers.cpp
+++ b/Engine/Source/Editor/MeshPaint/Private/MeshPaintHelpers.cpp
@@ -51,6 +51,8 @@
 #include "Async/ParallelFor.h"
 #include "Rendering/SkeletalMeshModel.h"
 
+extern void PropagateVertexPaintToAsset(USkeletalMesh* SkeletalMesh, int32 LODIndex);
+
 void MeshPaintHelpers::RemoveInstanceVertexColors(UObject* Obj)
 {
 	// Currently only static mesh component supports per instance vertex colors so only need to retrieve those and remove colors
@@ -912,6 +914,7 @@ void MeshPaintHelpers::FillSkeletalMeshVertexColors(USkeletalMeshComponent* Mesh
 			for (LODIndex = 0; LODIndex < NumLODs; ++LODIndex)
 			{
 				MeshPaintHelpers::SetColorDataForLOD(Mesh, LODIndex, FillColor, MaskColor);
+				PropagateVertexPaintToAsset(Mesh, LODIndex);
 			}
 			Mesh->InitResources();
 		}
@@ -1775,6 +1778,8 @@ void MeshPaintHelpers::ApplyVertexColorsToAllLODs(IMeshPaintGeometryAdapter& Geo
 				const FSkeletalMeshLODRenderData& BaseLOD = Resource->LODRenderData[0];
 				GeometryInfo.PreEdit();				
 
+				PropagateVertexPaintToAsset(Mesh, 0);
+
 				FBox BaseBounds(ForceInitToZero);
 
 				TArray<FPaintedMeshVertex> PaintedVertices;
@@ -1916,8 +1921,8 @@ void MeshPaintHelpers::ApplyVertexColorsToAllLODs(IMeshPaintGeometryAdapter& Geo
 							SrcLOD.Sections[SectionIndex].SoftVertices[SectionVertexIndex].Color = PointsToConsider[BestVertexIndex].Color;
 						}
 					}
+					PropagateVertexPaintToAsset(Mesh, LODIndex);
 				}
-				
 				GeometryInfo.PostEdit();
 			}
 		}
diff --git a/Engine/Source/Editor/MeshPaint/Private/MeshPaintSkeletalMeshAdapter.cpp b/Engine/Source/Editor/MeshPaint/Private/MeshPaintSkeletalMeshAdapter.cpp
--- a/Engine/Source/Editor/MeshPaint/Private/MeshPaintSkeletalMeshAdapter.cpp
+++ b/Engine/Source/Editor/MeshPaint/Private/MeshPaintSkeletalMeshAdapter.cpp
@@ -10,11 +10,99 @@
 #include "Rendering/SkeletalMeshModel.h"
 #include "Factories/FbxSkeletalMeshImportData.h"
 
+DEFINE_LOG_CATEGORY_STATIC(LogMeshPaintSkeletalMeshAdapter, Log, All);
+
 //////////////////////////////////////////////////////////////////////////
 // FMeshPaintGeometryAdapterForSkeletalMeshes
 
 FMeshPaintGeometryAdapterForSkeletalMeshes::FMeshToComponentMap FMeshPaintGeometryAdapterForSkeletalMeshes::MeshToComponentMap;
 
+
+//HACK for 4.24.2 we cannot change public API so we use this global function to remap and propagate the vertex color data to the imported model when the user release the mouse
+void PropagateVertexPaintToAsset(USkeletalMesh* SkeletalMesh, int32 LODIndex)
+{
+	struct FMatchFaceData
+	{
+		int32 SoftVerticeIndexes[3];
+	};
+
+	if (!SkeletalMesh || !SkeletalMesh->GetImportedModel() || !SkeletalMesh->GetImportedModel()->LODModels.IsValidIndex(LODIndex) 
+		|| SkeletalMesh->GetImportedModel()->LODModels[LODIndex].RawSkeletalMeshBulkData.IsEmpty() 
+		|| !SkeletalMesh->GetImportedModel()->LODModels[LODIndex].RawSkeletalMeshBulkData.IsBuildDataAvailable())
+	{
+		//We do not propagate vertex color for old asset
+		return;
+	}
+
+	auto GetMatchKey = [](const FVector& PositionA, const FVector& PositionB, const FVector& PositionC)->FSHAHash
+	{
+		FSHA1 SHA;
+		FSHAHash SHAHash;
+
+		SHA.Update((const uint8*)&PositionA, sizeof(FVector));
+		SHA.Update((const uint8*)&PositionB, sizeof(FVector));
+		SHA.Update((const uint8*)&PositionC, sizeof(FVector));
+		SHA.Final();
+		SHA.GetHash(&SHAHash.Hash[0]);
+
+		return SHAHash;
+	};
+
+	FSkeletalMeshLODModel& LODModel = SkeletalMesh->GetImportedModel()->LODModels[LODIndex];
+	const TArray<uint32>& SrcIndexBuffer = LODModel.IndexBuffer;
+
+	TArray<FSoftSkinVertex> SrcVertices;
+	LODModel.GetVertices(SrcVertices);
+	
+	FSkeletalMeshImportData ImportData;
+	LODModel.RawSkeletalMeshBulkData.LoadRawMesh(ImportData);
+
+	TMap<FSHAHash, FMatchFaceData> MatchTriangles;
+	MatchTriangles.Reserve(ImportData.Wedges.Num());
+
+	for (int32 IndexBufferIndex = 0, SrcIndexBufferNum = SrcIndexBuffer.Num(); IndexBufferIndex < SrcIndexBufferNum; IndexBufferIndex += 3)
+	{
+		FVector PositionA = SrcVertices[SrcIndexBuffer[IndexBufferIndex]].Position;
+		FVector PositionB = SrcVertices[SrcIndexBuffer[IndexBufferIndex + 1]].Position;
+		FVector PositionC = SrcVertices[SrcIndexBuffer[IndexBufferIndex + 2]].Position;
+
+		FSHAHash Key = GetMatchKey(PositionA, PositionB, PositionC);
+		FMatchFaceData MatchFaceData;
+		MatchFaceData.SoftVerticeIndexes[0] = SrcIndexBuffer[IndexBufferIndex];
+		MatchFaceData.SoftVerticeIndexes[1] = SrcIndexBuffer[IndexBufferIndex + 1];
+		MatchFaceData.SoftVerticeIndexes[2] = SrcIndexBuffer[IndexBufferIndex + 2];
+		MatchTriangles.Add(Key, MatchFaceData);
+	}
+
+	bool bAllVertexFound = true;
+	for (int32 FaceIndex = 0, FaceNum = ImportData.Faces.Num(); FaceIndex < FaceNum; ++FaceIndex)
+	{
+		const SkeletalMeshImportData::FTriangle& Triangle = ImportData.Faces[FaceIndex];
+		SkeletalMeshImportData::FVertex& WedgeA = ImportData.Wedges[Triangle.WedgeIndex[0]];
+		SkeletalMeshImportData::FVertex& WedgeB = ImportData.Wedges[Triangle.WedgeIndex[1]];
+		SkeletalMeshImportData::FVertex& WedgeC = ImportData.Wedges[Triangle.WedgeIndex[2]];
+		FVector PositionA = ImportData.Points[WedgeA.VertexIndex];
+		FVector PositionB = ImportData.Points[WedgeB.VertexIndex];
+		FVector PositionC = ImportData.Points[WedgeC.VertexIndex];
+
+		const FSHAHash Key = GetMatchKey(PositionA, PositionB, PositionC);
+		if (FMatchFaceData* MatchFaceData = MatchTriangles.Find(Key))
+		{
+			WedgeA.Color = SrcVertices[MatchFaceData->SoftVerticeIndexes[0]].Color;
+			WedgeB.Color = SrcVertices[MatchFaceData->SoftVerticeIndexes[1]].Color;
+			WedgeC.Color = SrcVertices[MatchFaceData->SoftVerticeIndexes[2]].Color;
+		}
+		else if(bAllVertexFound)
+		{
+			bAllVertexFound = false;
+			FString SkeletalMeshName(SkeletalMesh->GetName());
+			UE_LOG(LogMeshPaintSkeletalMeshAdapter, Warning, TEXT("Some vertex color data could not be applied to the %s SkeletalMesh asset."), *SkeletalMeshName);
+		}
+	}
+		
+	LODModel.RawSkeletalMeshBulkData.SaveRawMesh(ImportData);
+}
+
 bool FMeshPaintGeometryAdapterForSkeletalMeshes::Construct(UMeshComponent* InComponent, int32 InMeshLODIndex)
 {
 	SkeletalMeshComponent = Cast<USkeletalMeshComponent>(InComponent);
@@ -190,6 +278,7 @@ void FMeshPaintGeometryAdapterForSkeletalMeshes::OnRemoved()
 	{
 		return;
 	}
+	PropagateVertexPaintToAsset(ReferencedSkeletalMesh, MeshLODIndex);
 
 	// Remove a reference from the skeletal mesh map
 	FSkeletalMeshReferencers* SkeletalMeshReferencers = MeshToComponentMap.Find(ReferencedSkeletalMesh);
diff --git a/Engine/Source/Editor/MeshPaintMode/Private/PaintModePainter.cpp b/Engine/Source/Editor/MeshPaintMode/Private/PaintModePainter.cpp
--- a/Engine/Source/Editor/MeshPaintMode/Private/PaintModePainter.cpp
+++ b/Engine/Source/Editor/MeshPaintMode/Private/PaintModePainter.cpp
@@ -813,9 +813,16 @@ void FPaintModePainter::FinishPainting()
 	IMeshPainter::FinishPainting();	
 	FinishPaintingTexture();	
 
-	if (PaintSettings->PaintMode == EPaintMode::Vertices && !PaintSettings->VertexPaintSettings.bPaintOnSpecificLOD)
+	if (PaintSettings->PaintMode == EPaintMode::Vertices)
 	{
-		PropagateVertexColorsToLODs();
+		if (!PaintSettings->VertexPaintSettings.bPaintOnSpecificLOD)
+		{
+			PropagateVertexColorsToLODs();
+		}
+		else
+		{
+			Refresh();
+		}
 	}
 
 	UpdateCachedVertexDataSize();
