Fixing build

diff --git a/Engine/Source/Editor/MeshPaint/Private/MeshPaintSkeletalMeshAdapter.cpp b/Engine/Source/Editor/MeshPaint/Private/MeshPaintSkeletalMeshAdapter.cpp
--- a/Engine/Source/Editor/MeshPaint/Private/MeshPaintSkeletalMeshAdapter.cpp
+++ b/Engine/Source/Editor/MeshPaint/Private/MeshPaintSkeletalMeshAdapter.cpp
@@ -26,9 +26,7 @@ void PropagateVertexPaintToAsset(USkeletalMesh* SkeletalMesh, int32 LODIndex)
 		int32 SoftVerticeIndexes[3];
 	};
 
-	if (!SkeletalMesh || !SkeletalMesh->GetImportedModel() || !SkeletalMesh->GetImportedModel()->LODModels.IsValidIndex(LODIndex) 
-		|| SkeletalMesh->GetImportedModel()->LODModels[LODIndex].RawSkeletalMeshBulkData.IsEmpty() 
-		|| !SkeletalMesh->GetImportedModel()->LODModels[LODIndex].RawSkeletalMeshBulkData.IsBuildDataAvailable())
+	if (!SkeletalMesh || SkeletalMesh->IsLODImportedDataEmpty(LODIndex) || !SkeletalMesh->IsLODImportedDataBuildAvailable(LODIndex))
 	{
 		//We do not propagate vertex color for old asset
 		return;
@@ -55,7 +53,7 @@ void PropagateVertexPaintToAsset(USkeletalMesh* SkeletalMesh, int32 LODIndex)
 	LODModel.GetVertices(SrcVertices);
 	
 	FSkeletalMeshImportData ImportData;
-	LODModel.RawSkeletalMeshBulkData.LoadRawMesh(ImportData);
+	SkeletalMesh->LoadLODImportedData(LODIndex, ImportData);
 
 	TMap<FSHAHash, FMatchFaceData> MatchTriangles;
 	MatchTriangles.Reserve(ImportData.Wedges.Num());
@@ -100,7 +98,7 @@ void PropagateVertexPaintToAsset(USkeletalMesh* SkeletalMesh, int32 LODIndex)
 		}
 	}
 		
-	LODModel.RawSkeletalMeshBulkData.SaveRawMesh(ImportData);
+	SkeletalMesh->SaveLODImportedData(LODIndex, ImportData);
 }
 
 bool FMeshPaintGeometryAdapterForSkeletalMeshes::Construct(UMeshComponent* InComponent, int32 InMeshLODIndex)
