Cloth - Fixed camera movement when exiting the cloth painting mode.

diff --git a/Engine/Source/Editor/MeshPaint/Private/IMeshPaintMode.cpp b/Engine/Source/Editor/MeshPaint/Private/IMeshPaintMode.cpp
--- a/Engine/Source/Editor/MeshPaint/Private/IMeshPaintMode.cpp
+++ b/Engine/Source/Editor/MeshPaint/Private/IMeshPaintMode.cpp
@@ -223,6 +223,7 @@ bool IMeshPaintEdMode::InputKey( FEditorViewportClient* InViewportClient, FViewp
 		{
 			bHandled = true;
 			MeshPainter->FinishPainting();
+			InViewportClient->bLockFlightCamera = false;
 		}
 		else if( !MeshPainter->IsPainting() && bUserWantsPaint && !InViewportClient->IsMovingCamera())
 		{	
