Updated headers using UnrealCodeFixup to make sure dllstorage is on methods/staticvars instead of types. Part 5/n

diff --git a/Engine/Source/Editor/MaterialEditor/Public/MaterialEditorSettings.h b/Engine/Source/Editor/MaterialEditor/Public/MaterialEditorSettings.h
--- a/Engine/Source/Editor/MaterialEditor/Public/MaterialEditorSettings.h
+++ b/Engine/Source/Editor/MaterialEditor/Public/MaterialEditorSettings.h
@@ -26,7 +26,7 @@ enum class EBackgroundType : uint8
 };
 
 USTRUCT()
-struct MATERIALEDITOR_API FCheckerboardSettings
+struct FCheckerboardSettings
 {
 	GENERATED_BODY()
 	
@@ -64,7 +64,7 @@ inline bool operator!=(const FCheckerboardSettings& Lhs, const FCheckerboardSett
 }
 
 USTRUCT()
-struct MATERIALEDITOR_API FPreviewBackgroundSettings
+struct FPreviewBackgroundSettings
 {
 	GENERATED_BODY()
 	
@@ -99,8 +99,8 @@ struct MATERIALEDITOR_API FPreviewBackgroundSettings
 
 };
 
-UCLASS(config = EditorPerProjectUserSettings)
-class MATERIALEDITOR_API UMaterialEditorSettings : public UObject
+UCLASS(MinimalAPI, config = EditorPerProjectUserSettings)
+class UMaterialEditorSettings : public UObject
 {
 	GENERATED_BODY()
 
diff --git a/Engine/Source/Editor/MaterialEditor/Public/MaterialEditorUtilities.h b/Engine/Source/Editor/MaterialEditor/Public/MaterialEditorUtilities.h
--- a/Engine/Source/Editor/MaterialEditor/Public/MaterialEditorUtilities.h
+++ b/Engine/Source/Editor/MaterialEditor/Public/MaterialEditorUtilities.h
@@ -8,6 +8,8 @@
 #include "MaterialShared.h"
 #include "Materials/MaterialExpressionMaterialFunctionCall.h"
 
+#define UE_API MATERIALEDITOR_API
+
 class IMaterialEditor;
 class UMaterial;
 class UMaterialExpressionComment;
@@ -40,7 +42,7 @@ public:
 	FMaterialParameterInfo StackParameterInfo;
 };
 
-class MATERIALEDITOR_API FMaterialEditorUtilities
+class FMaterialEditorUtilities
 {
 public:
 	/**
@@ -52,7 +54,7 @@ public:
 	 * @param	bAutoSelect			If true, deselect all expressions and select the newly created one.
 	 * @param	bAutoAssignResource	If true, assign resources to new expression.
 	 */
-	static UMaterialExpression* CreateNewMaterialExpression(const class UEdGraph* Graph, UClass* NewExpressionClass, const UE::Slate::FDeprecateVector2DParameter& NodePos, bool bAutoSelect, bool bAutoAssignResource);
+	static UE_API UMaterialExpression* CreateNewMaterialExpression(const class UEdGraph* Graph, UClass* NewExpressionClass, const UE::Slate::FDeprecateVector2DParameter& NodePos, bool bAutoSelect, bool bAutoAssignResource);
 
 	/**
 	 * Creates a new material expression of the specified class on the material represented by a graph.
@@ -77,7 +79,7 @@ public:
 	 *
 	 * @return	UMaterialExpressionComment*	Newly created comment
 	 */
-	static UMaterialExpressionComposite* CreateNewMaterialExpressionComposite(const class UEdGraph* Graph, const FVector2D& NodePos);
+	static UE_API UMaterialExpressionComposite* CreateNewMaterialExpressionComposite(const class UEdGraph* Graph, const FVector2D& NodePos);
 
 	/**
 	 * Creates a new material expression comment on the material represented by a graph.
@@ -87,14 +89,14 @@ public:
 	 *
 	 * @return	UMaterialExpressionComment*	Newly created comment
 	 */
-	static UMaterialExpressionComment* CreateNewMaterialExpressionComment(const class UEdGraph* Graph, const FVector2D& NodePos);
+	static UE_API UMaterialExpressionComment* CreateNewMaterialExpressionComment(const class UEdGraph* Graph, const FVector2D& NodePos);
 
 	/**
 	 * Refreshes all material expression previews, regardless of whether or not realtime previews are enabled.
 	 *
 	 * @param	Graph	Graph representing a material.
 	 */
-	static void ForceRefreshExpressionPreviews(const class UEdGraph* Graph);
+	static UE_API void ForceRefreshExpressionPreviews(const class UEdGraph* Graph);
 
 	/**
 	 * Add the specified object to the list of selected objects
@@ -102,21 +104,21 @@ public:
 	 * @param	Graph	Graph representing a material.
 	 * @param	Obj		Object to add to selection.
 	 */
-	static void AddToSelection(const class UEdGraph* Graph, UMaterialExpression* Expression);
+	static UE_API void AddToSelection(const class UEdGraph* Graph, UMaterialExpression* Expression);
 
 	/**
 	 * Disconnects and removes the selected material graph nodes.
 	 *
 	 * @param	Graph	Graph representing a material.
 	 */
-	static void DeleteSelectedNodes(const class UEdGraph* Graph);
+	static UE_API void DeleteSelectedNodes(const class UEdGraph* Graph);
 
 	/**
 	 * Delete the specified nodes from the graph.
 	 * @param Graph Graph representing the material.
 	 * @param NodesToDelete Array of nodes to be removed from the graph.
 	*/
-	static void DeleteNodes(const class UEdGraph* Graph, const TArray<UEdGraphNode*>& NodesToDelete);
+	static UE_API void DeleteNodes(const class UEdGraph* Graph, const TArray<UEdGraphNode*>& NodesToDelete);
 
 
 	/**
@@ -124,28 +126,28 @@ public:
 	 *
 	 * @param	Graph	Graph representing a material or material function.
 	 */
-	static FText GetOriginalObjectName(const class UEdGraph* Graph);
+	static UE_API FText GetOriginalObjectName(const class UEdGraph* Graph);
 
 	/**
 	 * Re-links the material and updates its representation in the editor
 	 *
 	 * @param	Graph	Graph representing a material or material function.
 	 */
-	static void UpdateMaterialAfterGraphChange(const class UEdGraph* Graph);
+	static UE_API void UpdateMaterialAfterGraphChange(const class UEdGraph* Graph);
 
 	/** Mark the material as dirty. */
-	static void MarkMaterialDirty(const class UEdGraph* Graph);
+	static UE_API void MarkMaterialDirty(const class UEdGraph* Graph);
 
-	static void UpdateDetailView(const class UEdGraph* Graph);
+	static UE_API void UpdateDetailView(const class UEdGraph* Graph);
 
 	/** Can we paste to this graph? */
-	static bool CanPasteNodes(const class UEdGraph* Graph);
+	static UE_API bool CanPasteNodes(const class UEdGraph* Graph);
 
 	/** Perform paste on graph, at location */
-	static void  PasteNodesHere(class UEdGraph* Graph, const FVector2D& Location);
+	static UE_API void  PasteNodesHere(class UEdGraph* Graph, const FVector2D& Location);
 
 	/** Gets the number of selected nodes on this graph */
-	static int32 GetNumberOfSelectedNodes(const class UEdGraph* Graph);
+	static UE_API int32 GetNumberOfSelectedNodes(const class UEdGraph* Graph);
 
 	/**
 	 * Get all actions for placing material expressions
@@ -154,7 +156,7 @@ public:
 	 * @param bMaterialFunction				Whether we're dealing with a Material Function
 	 */
 	UE_DEPRECATED(5.8, "Should call version that provides a UMaterial or UMaterialFunction")
-	static void GetMaterialExpressionActions(FGraphActionMenuBuilder& ActionMenuBuilder, bool bMaterialFunction);
+	static UE_API void GetMaterialExpressionActions(FGraphActionMenuBuilder& ActionMenuBuilder, bool bMaterialFunction);
 
 	/**
 	 * Get all actions for placing material expressions
@@ -162,14 +164,14 @@ public:
 	 * @param [in,out]	ActionMenuBuilder	The output menu builder.
 	 * @param MaterialOrFunction			A UMaterial or UMaterialFunction object
 	 */
-	static void GetMaterialExpressionActions(FGraphActionMenuBuilder& ActionMenuBuilder, const UObject* MaterialOrFunction);
+	static UE_API void GetMaterialExpressionActions(FGraphActionMenuBuilder& ActionMenuBuilder, const UObject* MaterialOrFunction);
 
 	/**
 	 * Check whether an expression is in the favourites list
 	 *
 	 * @param	InExpression	The Expression we are checking
 	 */
-	static bool IsMaterialExpressionInFavorites(UMaterialExpression* InExpression);
+	static UE_API bool IsMaterialExpressionInFavorites(UMaterialExpression* InExpression);
 
 	/**
 	 * Get a preview material render proxy for an expression
@@ -179,14 +181,14 @@ public:
 	 *
 	 * @return	FMaterialRenderProxy*	Preview for this expression or NULL
 	 */
-	static FMaterialRenderProxy* GetExpressionPreview(const class UEdGraph* Graph, UMaterialExpression* InExpression);
+	static UE_API FMaterialRenderProxy* GetExpressionPreview(const class UEdGraph* Graph, UMaterialExpression* InExpression);
 
 	/**
 	 * Updates the material editor search results
 	 *
 	 * @param	Graph			Graph representing material
 	 */
-	static void UpdateSearchResults(const class UEdGraph* Graph);
+	static UE_API void UpdateSearchResults(const class UEdGraph* Graph);
 
 	/////////////////////////////////////////////////////
 	// Static functions moved from SMaterialEditorCanvas
@@ -198,10 +200,10 @@ public:
 	 * @param	MaterialInstance	The material instance that contains all parameter overrides.
 	 * @param	VisisbleExpressions	The array that will contain the name's of the visible parameter expressions.
 	 */
-	static void GetVisibleMaterialParameters(const UMaterial *Material, UMaterialInstance *MaterialInstance, TArray<FMaterialParameterInfo> &VisibleExpressions);
+	static UE_API void GetVisibleMaterialParameters(const UMaterial *Material, UMaterialInstance *MaterialInstance, TArray<FMaterialParameterInfo> &VisibleExpressions);
 
 	/** Finds an input in the passed in array with a matching Id. */
-	static const FFunctionExpressionInput* FindInputById(const UMaterialExpressionFunctionInput* InputExpression, const TArray<FFunctionExpressionInput>& Inputs);
+	static UE_API const FFunctionExpressionInput* FindInputById(const UMaterialExpressionFunctionInput* InputExpression, const TArray<FFunctionExpressionInput>& Inputs);
 
 	/** 
 	* Returns the value for a static switch material expression.
@@ -214,7 +216,7 @@ public:
 	* 
 	* @return	Returns true if a value for the switch expression is found, otherwise returns false.
 	*/
-	static bool GetStaticSwitchExpressionValue(UMaterialInstance* MaterialInstance, UMaterialExpression *SwitchValueExpression, bool& OutValue, FGuid& OutExpressionID, TArray<FGetVisibleMaterialParametersFunctionState*>& FunctionStack);
+	static UE_API bool GetStaticSwitchExpressionValue(UMaterialInstance* MaterialInstance, UMaterialExpression *SwitchValueExpression, bool& OutValue, FGuid& OutExpressionID, TArray<FGetVisibleMaterialParametersFunctionState*>& FunctionStack);
 
 	/**
 	 * Populates the specified material's Expressions array (eg if cooked out or old content).
@@ -222,39 +224,39 @@ public:
 	 *
 	 * @param	Material	Material we are initializing
 	 */
-	static void InitExpressions(UMaterial* Material);
+	static UE_API void InitExpressions(UMaterial* Material);
 
 	/**
 	 * Build the texture streaming data for a given material. Also update the parent hierarchy has only the delta are stored.
 	 *
 	 * @param	MaterialInterface	The material to update.
 	 */
-	static void BuildTextureStreamingData(UMaterialInterface* MaterialInterface);
+	static UE_API void BuildTextureStreamingData(UMaterialInterface* MaterialInterface);
 
 	/** Get IMaterialEditor for given object, if it exists */
-	static TSharedPtr<class IMaterialEditor> GetIMaterialEditorForObject(const UObject* ObjectToFocusOn);
+	static UE_API TSharedPtr<class IMaterialEditor> GetIMaterialEditorForObject(const UObject* ObjectToFocusOn);
 
 	/** Get IMaterialEditor to focus on given object, if it exists */
-	static void BringFocusAttentionOnObject(const UObject* ObjectToFocusOn);
+	static UE_API void BringFocusAttentionOnObject(const UObject* ObjectToFocusOn);
 
 
 	/** Commands for the Parents menu */
-	static void OnOpenMaterial(const FAssetData InMaterial);
-	static void OnOpenFunction(const FAssetData InFunction);
-	static void OnShowMaterialInContentBrowser(const FAssetData InMaterial);
-	static void OnShowFunctionInContentBrowser(const FAssetData InFunction);
+	static UE_API void OnOpenMaterial(const FAssetData InMaterial);
+	static UE_API void OnOpenFunction(const FAssetData InFunction);
+	static UE_API void OnShowMaterialInContentBrowser(const FAssetData InMaterial);
+	static UE_API void OnShowFunctionInContentBrowser(const FAssetData InFunction);
 
 	/**
 	 * Triggers a refresh or redraw for post process preview materials and material instances.  A refresh is useful when a material generating a UserSceneTexture
 	 * output is loaded or unloaded.  Other loaded materials that have UserSceneTexture inputs may include the material in question in their preview, which
 	 * necessitates a refresh.  A redraw is useful when debug settings change that may affect post process material viewports.
 	 */
-	static void RefreshPostProcessPreviewMaterials(UMaterialInterface* ExcludeMaterialInterface, bool bRedrawOnly = false);
+	static UE_API void RefreshPostProcessPreviewMaterials(UMaterialInterface* ExcludeMaterialInterface, bool bRedrawOnly = false);
 
 private:
 
-	static void OpenSelectedParentEditor(UMaterialFunctionInterface* InMaterialFunction);
-	static void OpenSelectedParentEditor(UMaterialInterface* InMaterialInterface);
+	static UE_API void OpenSelectedParentEditor(UMaterialFunctionInterface* InMaterialFunction);
+	static UE_API void OpenSelectedParentEditor(UMaterialInterface* InMaterialInterface);
 
 	/**
 	 * Recursively walks the expression tree and parses the visible expression branches.
@@ -263,7 +265,7 @@ private:
 	 * @param	MaterialInstance				The material instance that contains all parameter overrides.
 	 * @param	VisisbleExpressions				The array that will contain the name's of the visible parameter expressions.
 	 */
-	static void GetVisibleMaterialParametersFromExpression(
+	static UE_API void GetVisibleMaterialParametersFromExpression(
 		FMaterialExpressionKey MaterialExpressionKey, 
 		UMaterialInstance* MaterialInstance, 
 		TArray<FMaterialParameterInfo>& VisibleExpressions, 
@@ -278,7 +280,7 @@ private:
 	 * @param	bMaterialFunction	Whether we are building for a material function.
 	 */
 	UE_DEPRECATED(5.8, "Should call version that provides a UMaterial or UMaterialFunction")
-	static void AddMaterialExpressionCategory(FGraphActionMenuBuilder& ActionMenuBuilder, FText CategoryName, TArray<struct FMaterialExpression>* MaterialExpressions, bool bMaterialFunction);
+	static UE_API void AddMaterialExpressionCategory(FGraphActionMenuBuilder& ActionMenuBuilder, FText CategoryName, TArray<struct FMaterialExpression>* MaterialExpressions, bool bMaterialFunction);
 
 	/**
 	* Adds a category of Material Expressions to an action builder
@@ -288,7 +290,7 @@ private:
 	* @param	MaterialExpressions	List of Material Expressions in the category.
 	* @param	MaterialOrFunction	UMaterial or UMaterialFunction object
 	*/
-	static void AddMaterialExpressionCategory(FGraphActionMenuBuilder& ActionMenuBuilder, FText CategoryName, TArray<struct FMaterialExpression>* MaterialExpressions, const UObject* MaterialOrFunction);
+	static UE_API void AddMaterialExpressionCategory(FGraphActionMenuBuilder& ActionMenuBuilder, FText CategoryName, TArray<struct FMaterialExpression>* MaterialExpressions, const UObject* MaterialOrFunction);
 
 public:
 
@@ -300,8 +302,10 @@ public:
 	 * @param	TestDirection		Pin Direction we are testing.
 	 * @param	bMaterialFunction	Whether we are testing for a material function.
 	*/
-	static bool HasCompatibleConnection(UClass* ExpressionClass, uint32 TestType, EEdGraphPinDirection TestDirection, bool bMaterialFunction);
+	static UE_API bool HasCompatibleConnection(UClass* ExpressionClass, uint32 TestType, EEdGraphPinDirection TestDirection, bool bMaterialFunction);
 
 	/** Constructor */
 	FMaterialEditorUtilities() {}
 };
+
+#undef UE_API
diff --git a/Engine/Source/Editor/MaterialEditor/Public/MaterialNodes/SGraphNodeMaterialComment.h b/Engine/Source/Editor/MaterialEditor/Public/MaterialNodes/SGraphNodeMaterialComment.h
--- a/Engine/Source/Editor/MaterialEditor/Public/MaterialNodes/SGraphNodeMaterialComment.h
+++ b/Engine/Source/Editor/MaterialEditor/Public/MaterialNodes/SGraphNodeMaterialComment.h
@@ -7,21 +7,25 @@
 #include "SNodePanel.h"
 #include "Widgets/DeclarativeSyntaxSupport.h"
 
+#define UE_API MATERIALEDITOR_API
+
 class UMaterialGraphNode_Comment;
 
-class MATERIALEDITOR_API SGraphNodeMaterialComment : public SGraphNodeComment
+class SGraphNodeMaterialComment : public SGraphNodeComment
 {
 public:
 	SLATE_BEGIN_ARGS(SGraphNodeMaterialComment){}
 	SLATE_END_ARGS()
 
-	void Construct(const FArguments& InArgs, class UMaterialGraphNode_Comment* InNode);
+	UE_API void Construct(const FArguments& InArgs, class UMaterialGraphNode_Comment* InNode);
 
 	// SNodePanel::SNode interface
-	virtual void MoveTo(const FVector2f& NewPosition, FNodeSet& NodeFilter, bool bMarkDirty = true) override;
+	UE_API virtual void MoveTo(const FVector2f& NewPosition, FNodeSet& NodeFilter, bool bMarkDirty = true) override;
 	// End of SNodePanel::SNode interface
 
 private:
 	/** Cached material graph node pointer to avoid casting */
 	UMaterialGraphNode_Comment* CommentNode;
 };
+
+#undef UE_API
diff --git a/Engine/Source/Editor/MaterialEditor/Public/MaterialPropertyHelpers.h b/Engine/Source/Editor/MaterialEditor/Public/MaterialPropertyHelpers.h
--- a/Engine/Source/Editor/MaterialEditor/Public/MaterialPropertyHelpers.h
+++ b/Engine/Source/Editor/MaterialEditor/Public/MaterialPropertyHelpers.h
@@ -13,6 +13,8 @@
 #include "RenderUtils.h"
 #include "MaterialPropertyHelpers.generated.h"
 
+#define UE_API MATERIALEDITOR_API
+
 
 struct FAssetData;
 class IDetailGroup;
@@ -35,7 +37,7 @@ enum EStackDataType
 };
 
 USTRUCT()
-struct MATERIALEDITOR_API FSortedParamData
+struct FSortedParamData
 {
 	GENERATED_USTRUCT_BODY()
 
@@ -152,77 +154,78 @@ public:
    FMaterialInstanceBaseParameterDetails
 -----------------------------------------------------------------------------*/
 
-class MATERIALEDITOR_API FMaterialPropertyHelpers
+class FMaterialPropertyHelpers
 {
 public:
 	/** Returns true if the parameter is being overridden */
-	static bool IsOverriddenExpression(UDEditorParameterValue* Parameter);
+	static UE_API bool IsOverriddenExpression(UDEditorParameterValue* Parameter);
 	static bool IsOverriddenExpression(TObjectPtr<UDEditorParameterValue> Parameter) { return IsOverriddenExpression(Parameter.Get()); }
 	static bool IsOverriddenExpression(TWeakObjectPtr<UDEditorParameterValue> Parameter) { return IsOverriddenExpression(Parameter.Get()); }
-	static ECheckBoxState IsOverriddenExpressionCheckbox(UDEditorParameterValue* Parameter);
-	static bool UsesCustomPrimitiveData(UDEditorParameterValue* Parameter);
+	static UE_API ECheckBoxState IsOverriddenExpressionCheckbox(UDEditorParameterValue* Parameter);
+	static UE_API bool UsesCustomPrimitiveData(UDEditorParameterValue* Parameter);
 
 	/** Gets the expression description of this parameter from the base material */
-	static	FText GetParameterExpressionDescription(UDEditorParameterValue* Parameter, UObject* MaterialEditorInstance);
+	static	UE_API FText GetParameterExpressionDescription(UDEditorParameterValue* Parameter, UObject* MaterialEditorInstance);
 	
-	static FText GetParameterTooltip(UDEditorParameterValue* Parameter, UObject* MaterialEditorInstance);
+	static UE_API FText GetParameterTooltip(UDEditorParameterValue* Parameter, UObject* MaterialEditorInstance);
 	/**
 	 * Called when a parameter is overridden;
 	 */
-	static void OnOverrideParameter(bool NewValue, UDEditorParameterValue* Parameter, UMaterialEditorInstanceConstant* MaterialEditorInstance);
+	static UE_API void OnOverrideParameter(bool NewValue, UDEditorParameterValue* Parameter, UMaterialEditorInstanceConstant* MaterialEditorInstance);
 
-	static EVisibility ShouldShowExpression(UDEditorParameterValue* Parameter, UMaterialEditorInstanceConstant* MaterialEditorInstance, FGetShowHiddenParameters ShowHiddenDelegate);
+	static UE_API EVisibility ShouldShowExpression(UDEditorParameterValue* Parameter, UMaterialEditorInstanceConstant* MaterialEditorInstance, FGetShowHiddenParameters ShowHiddenDelegate);
 	static EVisibility ShouldShowExpression(TObjectPtr<UDEditorParameterValue> Parameter, UMaterialEditorInstanceConstant* MaterialEditorInstance, FGetShowHiddenParameters ShowHiddenDelegate) { return ShouldShowExpression(Parameter.Get(), MaterialEditorInstance, ShowHiddenDelegate); }
 
 	/** Generic material property reset to default implementation.  Resets Parameter to default */
-	static void ResetToDefault(UDEditorParameterValue* Parameter, UMaterialEditorInstanceConstant* MaterialEditorInstance);
+	static UE_API void ResetToDefault(UDEditorParameterValue* Parameter, UMaterialEditorInstanceConstant* MaterialEditorInstance);
 	static void ResetToDefault(TObjectPtr<UDEditorParameterValue> Parameter, UMaterialEditorInstanceConstant* MaterialEditorInstance) { ResetToDefault(Parameter.Get(), MaterialEditorInstance); }
-	static bool ShouldShowResetToDefault(UDEditorParameterValue* Parameter, UMaterialEditorInstanceConstant* MaterialEditorInstance);
+	static UE_API bool ShouldShowResetToDefault(UDEditorParameterValue* Parameter, UMaterialEditorInstanceConstant* MaterialEditorInstance);
 	static bool ShouldShowResetToDefault(TObjectPtr<UDEditorParameterValue> Parameter, UMaterialEditorInstanceConstant* MaterialEditorInstance) { return ShouldShowResetToDefault(Parameter.Get(), MaterialEditorInstance); }
 	
 	/** Specific resets for layer and blend asses */
-	static void ResetLayerAssetToDefault(UDEditorParameterValue* InParameter, TEnumAsByte<EMaterialParameterAssociation> InAssociation, int32 Index, UMaterialEditorInstanceConstant* MaterialEditorInstance);
+	static UE_API void ResetLayerAssetToDefault(UDEditorParameterValue* InParameter, TEnumAsByte<EMaterialParameterAssociation> InAssociation, int32 Index, UMaterialEditorInstanceConstant* MaterialEditorInstance);
 	/** If reset to default button should show for a layer or blend asset*/
-	static bool ShouldLayerAssetShowResetToDefault(TSharedPtr<FSortedParamData> InParameterData, UMaterialInstanceConstant* InMaterialInstance);
+	static UE_API bool ShouldLayerAssetShowResetToDefault(TSharedPtr<FSortedParamData> InParameterData, UMaterialInstanceConstant* InMaterialInstance);
 	static bool ShouldLayerAssetShowResetToDefault(TSharedPtr<FSortedParamData> InParameterData, TObjectPtr<UMaterialInstanceConstant> InMaterialInstance) { return ShouldLayerAssetShowResetToDefault(InParameterData, InMaterialInstance.Get()); }
 
-	static void OnMaterialLayerAssetChanged(const struct FAssetData& InAssetData, int32 Index, EMaterialParameterAssociation MaterialType, TSharedPtr<class IPropertyHandle> InHandle, FMaterialLayersFunctions* InMaterialFunction);
+	static UE_API void OnMaterialLayerAssetChanged(const struct FAssetData& InAssetData, int32 Index, EMaterialParameterAssociation MaterialType, TSharedPtr<class IPropertyHandle> InHandle, FMaterialLayersFunctions* InMaterialFunction);
 
-	static bool FilterLayerAssets(const struct FAssetData& InAssetData, FMaterialLayersFunctions* LayerFunction, EMaterialParameterAssociation MaterialType, int32 Index);
+	static UE_API bool FilterLayerAssets(const struct FAssetData& InAssetData, FMaterialLayersFunctions* LayerFunction, EMaterialParameterAssociation MaterialType, int32 Index);
 
-	static FReply OnClickedSaveNewMaterialInstance(class UMaterialInterface* Object, UObject* EditorObject);
+	static UE_API FReply OnClickedSaveNewMaterialInstance(class UMaterialInterface* Object, UObject* EditorObject);
 	static FReply OnClickedSaveNewMaterialInstance(TObjectPtr<class UMaterialInterface> Object, UObject* EditorObject) { return OnClickedSaveNewMaterialInstance(Object.Get(), EditorObject); }
 
-	static void CopyMaterialToInstance(class UMaterialInstanceConstant* ChildInstance, TArray<FEditorParameterGroup> &ParameterGroups);
-	static void TransitionAndCopyParameters(class UMaterialInstanceConstant* ChildInstance, TArray<FEditorParameterGroup> &ParameterGroups, bool bForceCopy = false);
-	static FReply OnClickedSaveNewFunctionInstance(class UMaterialFunctionInterface* Object, class UMaterialInterface* PreviewMaterial, UObject* EditorObject);
-	static FReply OnClickedSaveNewLayerInstance(class UMaterialFunctionInterface* Object, TSharedPtr<FSortedParamData> InSortedData);
+	static UE_API void CopyMaterialToInstance(class UMaterialInstanceConstant* ChildInstance, TArray<FEditorParameterGroup> &ParameterGroups);
+	static UE_API void TransitionAndCopyParameters(class UMaterialInstanceConstant* ChildInstance, TArray<FEditorParameterGroup> &ParameterGroups, bool bForceCopy = false);
+	static UE_API FReply OnClickedSaveNewFunctionInstance(class UMaterialFunctionInterface* Object, class UMaterialInterface* PreviewMaterial, UObject* EditorObject);
+	static UE_API FReply OnClickedSaveNewLayerInstance(class UMaterialFunctionInterface* Object, TSharedPtr<FSortedParamData> InSortedData);
 
-	static void GetVectorChannelMaskComboBoxStrings(TArray<TSharedPtr<FString>>& OutComboBoxStrings, TArray<TSharedPtr<class SToolTip>>& OutToolTips, TArray<bool>& OutRestrictedItems);
-	static FString GetVectorChannelMaskValue(UDEditorParameterValue* InParameter);
+	static UE_API void GetVectorChannelMaskComboBoxStrings(TArray<TSharedPtr<FString>>& OutComboBoxStrings, TArray<TSharedPtr<class SToolTip>>& OutToolTips, TArray<bool>& OutRestrictedItems);
+	static UE_API FString GetVectorChannelMaskValue(UDEditorParameterValue* InParameter);
 	static FString GetVectorChannelMaskValue(TObjectPtr<UDEditorParameterValue> InParameter) { return GetVectorChannelMaskValue(InParameter.Get()); }
-	static void SetVectorChannelMaskValue(const FString& StringValue, TSharedPtr<IPropertyHandle> PropertyHandle, UDEditorParameterValue* InParameter, UObject* MaterialEditorInstance);
+	static UE_API void SetVectorChannelMaskValue(const FString& StringValue, TSharedPtr<IPropertyHandle> PropertyHandle, UDEditorParameterValue* InParameter, UObject* MaterialEditorInstance);
 	static void SetVectorChannelMaskValue(const FString& StringValue, TSharedPtr<IPropertyHandle> PropertyHandle, TObjectPtr<UDEditorParameterValue> InParameter, UObject* MaterialEditorInstance) { SetVectorChannelMaskValue(StringValue, PropertyHandle, InParameter.Get(), MaterialEditorInstance); }
 
-	static TArray<class UFactory*> GetAssetFactories(EMaterialParameterAssociation AssetType);
+	static UE_API TArray<class UFactory*> GetAssetFactories(EMaterialParameterAssociation AssetType);
 	/**
 	*  Returns group for parameter. Creates one if needed.
 	*
 	* @param ParameterGroup		Name to be looked for.
 	*/
-	static FEditorParameterGroup&  GetParameterGroup(class UMaterial* InMaterial, FName& ParameterGroup, TArray<FEditorParameterGroup>& ParameterGroups);
+	static UE_API FEditorParameterGroup&  GetParameterGroup(class UMaterial* InMaterial, FName& ParameterGroup, TArray<FEditorParameterGroup>& ParameterGroups);
 
-	static TSharedRef<SWidget> MakeStackReorderHandle(TSharedPtr<IDraggableItem> InOwningStack);
+	static UE_API TSharedRef<SWidget> MakeStackReorderHandle(TSharedPtr<IDraggableItem> InOwningStack);
 
-	static bool OnShouldSetCurveAsset(const FAssetData& AssetData, TSoftObjectPtr<UCurveLinearColorAtlas> InAtlas);
-	static bool OnShouldFilterCurveAsset(const FAssetData& AssetData, TSoftObjectPtr<UCurveLinearColorAtlas> InAtlas);
-	static void SetPositionFromCurveAsset(const FAssetData& AssetData, TSoftObjectPtr<UCurveLinearColorAtlas> InAtlas, class UDEditorScalarParameterValue* InParameter, TSharedPtr<IPropertyHandle> PropertyHandle, UObject* MaterialEditorInstance);
+	static UE_API bool OnShouldSetCurveAsset(const FAssetData& AssetData, TSoftObjectPtr<UCurveLinearColorAtlas> InAtlas);
+	static UE_API bool OnShouldFilterCurveAsset(const FAssetData& AssetData, TSoftObjectPtr<UCurveLinearColorAtlas> InAtlas);
+	static UE_API void SetPositionFromCurveAsset(const FAssetData& AssetData, TSoftObjectPtr<UCurveLinearColorAtlas> InAtlas, class UDEditorScalarParameterValue* InParameter, TSharedPtr<IPropertyHandle> PropertyHandle, UObject* MaterialEditorInstance);
 
-	static void ResetCurveToDefault(UDEditorParameterValue* Parameter, UMaterialEditorInstanceConstant* MaterialEditorInstance);
+	static UE_API void ResetCurveToDefault(UDEditorParameterValue* Parameter, UMaterialEditorInstanceConstant* MaterialEditorInstance);
 	static void ResetCurveToDefault(TObjectPtr<UDEditorParameterValue> Parameter, UMaterialEditorInstanceConstant* MaterialEditorInstance) { ResetCurveToDefault(Parameter.Get(), MaterialEditorInstance); }
 
-	static FText LayerID;
-	static FText BlendID;
-	static FName LayerParamName;
+	static UE_API FText LayerID;
+	static UE_API FText BlendID;
+	static UE_API FName LayerParamName;
 };
 
+#undef UE_API
diff --git a/Engine/Source/Editor/MaterialEditor/Public/MaterialStatsCommon.h b/Engine/Source/Editor/MaterialEditor/Public/MaterialStatsCommon.h
--- a/Engine/Source/Editor/MaterialEditor/Public/MaterialStatsCommon.h
+++ b/Engine/Source/Editor/MaterialEditor/Public/MaterialStatsCommon.h
@@ -6,8 +6,10 @@
 #include "MaterialShared.h"
 #include "Styling/SlateColor.h"
 
+#define UE_API MATERIALEDITOR_API
+
 /** custom resource material class used to mark the resource as used for shader stats extraction */
-class MATERIALEDITOR_API FMaterialResourceStats : public FMaterialResource
+class FMaterialResourceStats : public FMaterialResource
 {
 public:
 	FMaterialResourceStats() {}
@@ -30,7 +32,7 @@ public:
 	virtual bool RequiresSynchronousCompilation() const override { return false; }
 
 	/** this will enable shader source extraction and pass paths to (eventual) offline shader compilers */
-	virtual void SetupExtraCompilationSettings(const EShaderPlatform Platform, FExtraShaderCompilerSettings& Settings) const override;
+	UE_API virtual void SetupExtraCompilationSettings(const EShaderPlatform Platform, FExtraShaderCompilerSettings& Settings) const override;
 };
 
 /** enumeration used to group shader platforms */
@@ -74,7 +76,7 @@ enum class ERepresentativeShader
 class FMaterialStatsUtils
 {
 public:
-	struct MATERIALEDITOR_API FShaderInstructionsInfo
+	struct FShaderInstructionsInfo
 	{
 		ERepresentativeShader ShaderType;
 		FString ShaderDescription;
@@ -82,7 +84,7 @@ public:
 		int32 InstructionCount;
 	};
 
-	struct MATERIALEDITOR_API FRepresentativeShaderInfo
+	struct FRepresentativeShaderInfo
 	{
 		ERepresentativeShader ShaderType;
 		FName ShaderName;
@@ -125,3 +127,5 @@ public:
 	MATERIALEDITOR_API static void GetPlatformOfflineCompilerSettings(const EShaderPlatform ShaderPlatform, FExtraShaderCompilerSettings& SCSettings);
 	MATERIALEDITOR_API static bool PlatformNeedsOfflineCompiler(const EShaderPlatform ShaderPlatform);
 };
+
+#undef UE_API
diff --git a/Engine/Source/Editor/MergeActors/Public/IMergeActorsModule.h b/Engine/Source/Editor/MergeActors/Public/IMergeActorsModule.h
--- a/Engine/Source/Editor/MergeActors/Public/IMergeActorsModule.h
+++ b/Engine/Source/Editor/MergeActors/Public/IMergeActorsModule.h
@@ -26,15 +26,15 @@ public:
 	/**
 	 * Register an IMergeActorsTool with the module, passing ownership to it
 	 */
-	MERGEACTORS_API virtual bool RegisterMergeActorsTool(TUniquePtr<IMergeActorsTool> Tool) = 0;
+	virtual bool RegisterMergeActorsTool(TUniquePtr<IMergeActorsTool> Tool) = 0;
 
 	/**
 	 * Unregister an IMergeActorsTool with the module
 	 */
-	MERGEACTORS_API virtual bool UnregisterMergeActorsTool(IMergeActorsTool* Tool) = 0;
+	virtual bool UnregisterMergeActorsTool(IMergeActorsTool* Tool) = 0;
 
 	/**
 	 * Get currently registered list of MergeActors tools
 	 */
-	MERGEACTORS_API virtual void GetRegisteredMergeActorsTools(TArray<IMergeActorsTool*>& OutTools) = 0;
+	virtual void GetRegisteredMergeActorsTools(TArray<IMergeActorsTool*>& OutTools) = 0;
 };
diff --git a/Engine/Source/Editor/MeshPaint/Public/BaseMeshPaintGeometryAdapter.h b/Engine/Source/Editor/MeshPaint/Public/BaseMeshPaintGeometryAdapter.h
--- a/Engine/Source/Editor/MeshPaint/Public/BaseMeshPaintGeometryAdapter.h
+++ b/Engine/Source/Editor/MeshPaint/Public/BaseMeshPaintGeometryAdapter.h
@@ -11,24 +11,26 @@
 #include "TMeshPaintOctree.h"
 #include "Templates/UniquePtr.h"
 
+#define UE_API MESHPAINT_API
+
 /** Base mesh paint geometry adapter, handles basic sphere intersection using a Octree */
-class MESHPAINT_API FBaseMeshPaintGeometryAdapter : public IMeshPaintGeometryAdapter
+class FBaseMeshPaintGeometryAdapter : public IMeshPaintGeometryAdapter
 {
 public:
 	/** Start IMeshPaintGeometryAdapter Overrides */
-	virtual bool Initialize() override;
-	virtual const TArray<FVector>& GetMeshVertices() const override;
-	virtual const TArray<uint32>& GetMeshIndices() const override;
-	virtual void GetVertexPosition(int32 VertexIndex, FVector& OutVertex) const override;
-	virtual TArray<uint32> SphereIntersectTriangles(const float ComponentSpaceSquaredBrushRadius, const FVector& ComponentSpaceBrushPosition, const FVector& ComponentSpaceCameraPosition, const bool bOnlyFrontFacing) const override;
-	virtual void GetInfluencedVertexIndices(const float ComponentSpaceSquaredBrushRadius, const FVector& ComponentSpaceBrushPosition, const FVector& ComponentSpaceCameraPosition, const bool bOnlyFrontFacing, TSet<int32> &InfluencedVertices) const override;
-	virtual void GetInfluencedVertexData(const float ComponentSpaceSquaredBrushRadius, const FVector& ComponentSpaceBrushPosition, const FVector& ComponentSpaceCameraPosition, const bool bOnlyFrontFacing, TArray<TPair<int32, FVector>>& OutData) const override;
-	virtual TArray<FVector> SphereIntersectVertices(const float ComponentSpaceSquaredBrushRadius, const FVector& ComponentSpaceBrushPosition, const FVector& ComponentSpaceCameraPosition, const bool bOnlyFrontFacing) const override;
+	UE_API virtual bool Initialize() override;
+	UE_API virtual const TArray<FVector>& GetMeshVertices() const override;
+	UE_API virtual const TArray<uint32>& GetMeshIndices() const override;
+	UE_API virtual void GetVertexPosition(int32 VertexIndex, FVector& OutVertex) const override;
+	UE_API virtual TArray<uint32> SphereIntersectTriangles(const float ComponentSpaceSquaredBrushRadius, const FVector& ComponentSpaceBrushPosition, const FVector& ComponentSpaceCameraPosition, const bool bOnlyFrontFacing) const override;
+	UE_API virtual void GetInfluencedVertexIndices(const float ComponentSpaceSquaredBrushRadius, const FVector& ComponentSpaceBrushPosition, const FVector& ComponentSpaceCameraPosition, const bool bOnlyFrontFacing, TSet<int32> &InfluencedVertices) const override;
+	UE_API virtual void GetInfluencedVertexData(const float ComponentSpaceSquaredBrushRadius, const FVector& ComponentSpaceBrushPosition, const FVector& ComponentSpaceCameraPosition, const bool bOnlyFrontFacing, TArray<TPair<int32, FVector>>& OutData) const override;
+	UE_API virtual TArray<FVector> SphereIntersectVertices(const float ComponentSpaceSquaredBrushRadius, const FVector& ComponentSpaceBrushPosition, const FVector& ComponentSpaceCameraPosition, const bool bOnlyFrontFacing) const override;
 	/** End IMeshPaintGeometryAdapter Overrides */
 
 	virtual bool InitializeVertexData() = 0;
 protected:
-	bool BuildOctree();
+	UE_API bool BuildOctree();
 protected:
 	/** Index and Vertex data populated by derived classes in InitializeVertexData */
 	TArray<FVector> MeshVertices;
@@ -36,3 +38,5 @@ protected:
 	/** Octree used for reducing the cost of sphere intersecting with triangles / vertices */
 	TUniquePtr<FMeshPaintTriangleOctree> MeshTriOctree;
 };
+
+#undef UE_API
diff --git a/Engine/Source/Editor/MeshPaint/Public/IMeshPaintGeometryAdapter.h b/Engine/Source/Editor/MeshPaint/Public/IMeshPaintGeometryAdapter.h
--- a/Engine/Source/Editor/MeshPaint/Public/IMeshPaintGeometryAdapter.h
+++ b/Engine/Source/Editor/MeshPaint/Public/IMeshPaintGeometryAdapter.h
@@ -4,6 +4,8 @@
 
 #include "CoreMinimal.h"
 
+#define UE_API MESHPAINT_API
+
 class FReferenceCollector;
 class UMeshComponent;
 class UTexture;
@@ -11,7 +13,7 @@ class UTexture;
 /**
  * Interface for a class to provide mesh painting support for a subclass of UMeshComponent
  */
-class MESHPAINT_API IMeshPaintGeometryAdapter
+class IMeshPaintGeometryAdapter
 {
 public:
 	/** Constructs the adapter for a specific LOD index using the mesh component */
@@ -78,10 +80,10 @@ public:
 	virtual void PostEdit() = 0;
 
 	/** Default functionality for applying or remove a texture override */
-	static void DefaultApplyOrRemoveTextureOverride(UMeshComponent* InMeshComponent, UTexture* SourceTexture, UTexture* OverrideTexture);
+	static UE_API void DefaultApplyOrRemoveTextureOverride(UMeshComponent* InMeshComponent, UTexture* SourceTexture, UTexture* OverrideTexture);
 
 	/** Default functionality for querying textures from a mesh component which are suitable for texture painting */
-	static void DefaultQueryPaintableTextures(int32 MaterialIndex, const UMeshComponent* MeshComponent, int32& OutDefaultIndex, TArray<struct FPaintableTexture>& InOutTextureList);
+	static UE_API void DefaultQueryPaintableTextures(int32 MaterialIndex, const UMeshComponent* MeshComponent, int32& OutDefaultIndex, TArray<struct FPaintableTexture>& InOutTextureList);
 
 	/** Returns the triangle indices which intersect with the given sphere */
 	virtual TArray<uint32> SphereIntersectTriangles(const float ComponentSpaceSquaredBrushRadius, const FVector& ComponentSpaceBrushPosition, const FVector& ComponentSpaceCameraPosition, const bool bOnlyFrontFacing) const = 0;
@@ -94,4 +96,6 @@ public:
 
 	/** Returns the vertex positions which intersect the given sphere */
 	virtual TArray<FVector> SphereIntersectVertices(const float ComponentSpaceSquaredBrushRadius, const FVector& ComponentSpaceBrushPosition, const FVector& ComponentSpaceCameraPosition, const bool bOnlyFrontFacing) const = 0;
-};
\ No newline at end of file
+};
+
+#undef UE_API
diff --git a/Engine/Source/Editor/MeshPaint/Public/IMeshPaintMode.h b/Engine/Source/Editor/MeshPaint/Public/IMeshPaintMode.h
--- a/Engine/Source/Editor/MeshPaint/Public/IMeshPaintMode.h
+++ b/Engine/Source/Editor/MeshPaint/Public/IMeshPaintMode.h
@@ -6,6 +6,8 @@
 #include "InputCoreTypes.h"
 #include "EdMode.h"
 
+#define UE_API MESHPAINT_API
+
 class FEditorViewportClient;
 class FPrimitiveDrawInterface;
 class FObjectPostSaveContext;
@@ -20,20 +22,20 @@ struct FAssetData;
 /**
  * Mesh Paint editor mode
  */
-class MESHPAINT_API IMeshPaintEdMode : public FEdMode
+class IMeshPaintEdMode : public FEdMode
 {
 public:
 	/** Constructor */
-	IMeshPaintEdMode();
+	UE_API IMeshPaintEdMode();
 
 	/** Destructor */
-	virtual ~IMeshPaintEdMode();
+	UE_API virtual ~IMeshPaintEdMode();
 
 	virtual void Initialize() = 0;
 	virtual TSharedPtr<class FModeToolkit> GetToolkit() = 0;
 
 	/** FGCObject interface */
-	virtual void AddReferencedObjects( FReferenceCollector& Collector ) override;
+	UE_API virtual void AddReferencedObjects( FReferenceCollector& Collector ) override;
 	virtual FString GetReferencerName() const override
 	{
 		return "IMeshPaintEdMode";
@@ -41,49 +43,49 @@ public:
 
 	// FEdMode interface
 	virtual bool UsesToolkits() const override { return true; }
-	virtual void Enter() override;
-	virtual void Exit() override;
-	virtual bool ProcessCapturedMouseMoves(FEditorViewportClient* InViewportClient, FViewport* InViewport, const TArrayView<FIntPoint>& CapturedMouseMoves) override;
+	UE_API virtual void Enter() override;
+	UE_API virtual void Exit() override;
+	UE_API virtual bool ProcessCapturedMouseMoves(FEditorViewportClient* InViewportClient, FViewport* InViewport, const TArrayView<FIntPoint>& CapturedMouseMoves) override;
 	virtual bool StartTracking(FEditorViewportClient* InViewportClient, FViewport* InViewport) override { return true; }
-	virtual bool EndTracking(FEditorViewportClient* InViewportClient, FViewport* InViewport) override;
-	virtual bool InputKey( FEditorViewportClient* InViewportClient, FViewport* InViewport, FKey InKey, EInputEvent InEvent ) override;
-	virtual void PostUndo() override;
-	virtual void Render( const FSceneView* View, FViewport* Viewport, FPrimitiveDrawInterface* PDI ) override;
-	virtual bool Select( AActor* InActor, bool bInSelected ) override;
-	virtual void ActorSelectionChangeNotify() override;
+	UE_API virtual bool EndTracking(FEditorViewportClient* InViewportClient, FViewport* InViewport) override;
+	UE_API virtual bool InputKey( FEditorViewportClient* InViewportClient, FViewport* InViewport, FKey InKey, EInputEvent InEvent ) override;
+	UE_API virtual void PostUndo() override;
+	UE_API virtual void Render( const FSceneView* View, FViewport* Viewport, FPrimitiveDrawInterface* PDI ) override;
+	UE_API virtual bool Select( AActor* InActor, bool bInSelected ) override;
+	UE_API virtual void ActorSelectionChangeNotify() override;
 	virtual bool AllowWidgetMove() override { return false; }
 	virtual bool ShouldDrawWidget() const override { return false; }
 	virtual bool UsesTransformWidget() const override { return false; }
-	virtual void Tick( FEditorViewportClient* ViewportClient, float DeltaTime ) override;
-	virtual bool ProcessEditDelete() override;
+	UE_API virtual void Tick( FEditorViewportClient* ViewportClient, float DeltaTime ) override;
+	UE_API virtual bool ProcessEditDelete() override;
 	// End of FEdMode interface
 
 	/** Returns the mesh painter for this mode */
-	IMeshPainter* GetMeshPainter();
+	UE_API IMeshPainter* GetMeshPainter();
 private:
 	/** Called prior to saving a level */
-	void OnPreSaveWorld(UWorld* World, FObjectPreSaveContext ObjectSaveContext);
+	UE_API void OnPreSaveWorld(UWorld* World, FObjectPreSaveContext ObjectSaveContext);
 
 	/** Called after saving a level */
-	void OnPostSaveWorld(UWorld* World, FObjectPostSaveContext ObjectSaveContext);
+	UE_API void OnPostSaveWorld(UWorld* World, FObjectPostSaveContext ObjectSaveContext);
 
 	/** Called when an asset has just been imported */
-	void OnPostImportAsset(UFactory* Factory, UObject* Object);
+	UE_API void OnPostImportAsset(UFactory* Factory, UObject* Object);
 
 	/** Called when an asset has just been reimported */
-	void OnPostReimportAsset(UObject* Object, bool bSuccess);
+	UE_API void OnPostReimportAsset(UObject* Object, bool bSuccess);
 
 	/** Called when an asset is deleted */
-	void OnAssetRemoved(const FAssetData& AssetData);
+	UE_API void OnAssetRemoved(const FAssetData& AssetData);
 
 	/** Called when the user presses a button on their motion controller device */
-	void OnVRAction( FEditorViewportClient& ViewportClient, class UViewportInteractor* Interactor,
+	UE_API void OnVRAction( FEditorViewportClient& ViewportClient, class UViewportInteractor* Interactor,
 	const struct FViewportActionKeyInput& Action, bool& bOutIsInputCaptured, bool& bWasHandled );
 
 	/** Called when rerunning a construction script causes objects to be replaced */
-	void OnObjectsReplaced(const TMap<UObject*, UObject*>& OldToNewInstanceMap);
+	UE_API void OnObjectsReplaced(const TMap<UObject*, UObject*>& OldToNewInstanceMap);
 
-	void OnResetViewMode();
+	UE_API void OnResetViewMode();
 private:	
 	/** When painting in VR, this is the hand index that we're painting with.  Otherwise INDEX_NONE. */
 	UViewportInteractor* PaintingWithInteractorInVR;
@@ -96,4 +98,6 @@ private:
 protected:
 	/** Painter used by this edit mode for applying paint actions */
 	IMeshPainter* MeshPainter;	
-};
\ No newline at end of file
+};
+
+#undef UE_API
diff --git a/Engine/Source/Editor/MeshPaint/Public/IMeshPainter.h b/Engine/Source/Editor/MeshPaint/Public/IMeshPainter.h
--- a/Engine/Source/Editor/MeshPaint/Public/IMeshPainter.h
+++ b/Engine/Source/Editor/MeshPaint/Public/IMeshPainter.h
@@ -7,6 +7,8 @@
 #include "Engine/EngineBaseTypes.h"
 #include "InputCoreTypes.h"
 
+#define UE_API MESHPAINT_API
+
 class UMeshPaintSettings;
 class UPaintBrushSettings;
 class UMeshComponent;
@@ -24,38 +26,38 @@ struct FHitResult;
 enum class EMeshPaintAction : uint8;
 
 /** Base class for creating a mesh painter, has basic functionality combined with IMeshPaintMode and requires painting/situation-specific information to do actual painting */
-class MESHPAINT_API IMeshPainter
+class IMeshPainter
 {
 public:
-	IMeshPainter();
-	virtual ~IMeshPainter();
+	UE_API IMeshPainter();
+	UE_API virtual ~IMeshPainter();
 
 	/** Renders ray widgets for the active viewport interactors */
 	virtual void Render(const FSceneView* View, FViewport* Viewport, FPrimitiveDrawInterface* PDI) = 0;
 
 	/** Base ticking functionality, keeps track of time and caches VR interactors */
-	virtual void Tick(FEditorViewportClient* ViewportClient, float DeltaTime);
+	UE_API virtual void Tick(FEditorViewportClient* ViewportClient, float DeltaTime);
 
 	/** Returns whether or not painting (of any sorts) is in progress*/
 	virtual bool IsPainting() { return bArePainting; }
 
 	/** Registers a base set of commands for altering UBrushSettings */
-	virtual void RegisterCommands(TSharedRef<FUICommandList> CommandList);
+	UE_API virtual void RegisterCommands(TSharedRef<FUICommandList> CommandList);
 
 	/** Unregisters a base set of commands */
-	virtual void UnregisterCommands(TSharedRef<FUICommandList> CommandList);
+	UE_API virtual void UnregisterCommands(TSharedRef<FUICommandList> CommandList);
 
 	/** 'Normal' painting functionality, called when the user tries to paint on a mesh using the mouse */
-	virtual bool Paint(FViewport* Viewport, const FVector& InCameraOrigin, const FVector& InRayOrigin, const FVector& InRayDirection);
+	UE_API virtual bool Paint(FViewport* Viewport, const FVector& InCameraOrigin, const FVector& InRayOrigin, const FVector& InRayDirection);
 
 	/** 'Normal' painting functionality, called when the user tries to paint on a mesh using the mouse */
-	virtual bool Paint(FViewport* Viewport, const FVector& InCameraOrigin, const TArrayView<TPair<FVector, FVector>>& Rays);
+	UE_API virtual bool Paint(FViewport* Viewport, const FVector& InCameraOrigin, const TArrayView<TPair<FVector, FVector>>& Rays);
 
 	/** Allows painter to act on specific key actions */
-	virtual bool InputKey(FEditorViewportClient* InViewportClient, FViewport* InViewport, FKey InKey, EInputEvent InEvent);
+	UE_API virtual bool InputKey(FEditorViewportClient* InViewportClient, FViewport* InViewport, FKey InKey, EInputEvent InEvent);
 
 	/** If painting is in progress will propagate changes and cleanup the painting state */
-	virtual void FinishPainting();
+	UE_API virtual void FinishPainting();
 	
 	/** Is called by the owning EdMode when an actor is de-selected in the viewport */
 	virtual void ActorSelected(AActor* Actor) = 0;
@@ -90,10 +92,10 @@ protected:
 	virtual bool PaintInternal(const FVector& InCameraOrigin, const TArrayView<TPair<FVector, FVector>>& Rays, EMeshPaintAction PaintAction, float PaintStrength) = 0;
 
 	/** Renders viewport interactor widget */
-	void RenderInteractors(const FSceneView* View, FViewport* Viewport, FPrimitiveDrawInterface* PDI, bool bRenderVertices, ESceneDepthPriorityGroup DepthGroup = SDPG_World);
+	UE_API void RenderInteractors(const FSceneView* View, FViewport* Viewport, FPrimitiveDrawInterface* PDI, bool bRenderVertices, ESceneDepthPriorityGroup DepthGroup = SDPG_World);
 
 	/** Functionality for rendering the interactor widget given its positional information and paint action */
-	virtual void RenderInteractorWidget(const FVector& InCameraOrigin, const FVector& InRayOrigin, const FVector& InRayDirection, FPrimitiveDrawInterface* PDI, EMeshPaintAction PaintAction, bool bRenderVertices, ESceneDepthPriorityGroup DepthGroup = SDPG_World);
+	UE_API virtual void RenderInteractorWidget(const FVector& InCameraOrigin, const FVector& InRayOrigin, const FVector& InRayDirection, FPrimitiveDrawInterface* PDI, EMeshPaintAction PaintAction, bool bRenderVertices, ESceneDepthPriorityGroup DepthGroup = SDPG_World);
 
 protected:
 	/** Flag for whether or not we are currently painting */
@@ -112,20 +114,22 @@ protected:
 
 protected:
 	/** Begins and keeps track of an Editor transaction with the given Description */
-	void BeginTransaction(const FText Description);
+	UE_API void BeginTransaction(const FText Description);
 	/** Ends the current transaction */
-	void EndTransaction();
+	UE_API void EndTransaction();
 	/** Painting transaction instance which is currently active */
 	FScopedTransaction* PaintTransaction;
 
 private:
 
 	/** Change the brush radius by a percentage, multiplied by the given value. */
-	void ChangeBrushRadius(float Multiplier);
+	UE_API void ChangeBrushRadius(float Multiplier);
 
 	/** Change the brush radius by one increment, multiplied by the given value */
-	void ChangeBrushStrength(float Multiplier);
+	UE_API void ChangeBrushStrength(float Multiplier);
 
 	/** Change the brush radius by one increment, multiplied by the given value */
-	void ChangeBrushFalloff(float Multiplier);
-};
\ No newline at end of file
+	UE_API void ChangeBrushFalloff(float Multiplier);
+};
+
+#undef UE_API
diff --git a/Engine/Source/Editor/MeshPaint/Public/MeshPaintAdapterFactory.h b/Engine/Source/Editor/MeshPaint/Public/MeshPaintAdapterFactory.h
--- a/Engine/Source/Editor/MeshPaint/Public/MeshPaintAdapterFactory.h
+++ b/Engine/Source/Editor/MeshPaint/Public/MeshPaintAdapterFactory.h
@@ -8,18 +8,22 @@
 #include "MeshPaintModule.h"
 #include "Templates/SharedPointer.h"
 
+#define UE_API MESHPAINT_API
+
 class FReferenceCollector;
 class IMeshPaintGeometryAdapterFactory;
 class UMeshComponent;
 
-class MESHPAINT_API FMeshPaintAdapterFactory
+class FMeshPaintAdapterFactory
 {
 public:
-	static TArray<TSharedPtr<IMeshPaintGeometryAdapterFactory>> FactoryList;
+	static UE_API TArray<TSharedPtr<IMeshPaintGeometryAdapterFactory>> FactoryList;
 
 public:
-	static TSharedPtr<class IMeshPaintGeometryAdapter> CreateAdapterForMesh(UMeshComponent* InComponent, int32 InPaintingMeshLODIndex);
-	static void InitializeAdapterGlobals();
-	static void AddReferencedObjectsGlobals(FReferenceCollector& Collector);
-	static void CleanupGlobals();
+	static UE_API TSharedPtr<class IMeshPaintGeometryAdapter> CreateAdapterForMesh(UMeshComponent* InComponent, int32 InPaintingMeshLODIndex);
+	static UE_API void InitializeAdapterGlobals();
+	static UE_API void AddReferencedObjectsGlobals(FReferenceCollector& Collector);
+	static UE_API void CleanupGlobals();
 };
+
+#undef UE_API
diff --git a/Engine/Source/Editor/MeshPaint/Public/MeshPaintHelpers.h b/Engine/Source/Editor/MeshPaint/Public/MeshPaintHelpers.h
--- a/Engine/Source/Editor/MeshPaint/Public/MeshPaintHelpers.h
+++ b/Engine/Source/Editor/MeshPaint/Public/MeshPaintHelpers.h
@@ -7,6 +7,8 @@
 #include "Engine/HitResult.h"
 #include "MeshPaintTypes.h"
 
+#define UE_API MESHPAINT_API
+
 class FMeshPaintParameters;
 class UVertexColorImportOptions;
 class UTexture2D;
@@ -41,114 +43,114 @@ struct FPerVertexPaintActionArgs
 DECLARE_DELEGATE_TwoParams(FPerVertexPaintAction, FPerVertexPaintActionArgs& /*Args*/, int32 /*VertexIndex*/);
 DECLARE_DELEGATE_ThreeParams(FPerTrianglePaintAction, IMeshPaintGeometryAdapter* /*Adapter*/, int32 /*TriangleIndex*/, const int32[3] /*Vertex Indices*/);
 
-class MESHPAINT_API MeshPaintHelpers
+class MeshPaintHelpers
 {
 public:
 	/** Removes vertex colors associated with the object */
-	static void RemoveInstanceVertexColors(UObject* Obj);
+	static UE_API void RemoveInstanceVertexColors(UObject* Obj);
 
 	/** Removes vertex colors associated with the static mesh component */
-	static void RemoveComponentInstanceVertexColors(UStaticMeshComponent* StaticMeshComponent);
+	static UE_API void RemoveComponentInstanceVertexColors(UStaticMeshComponent* StaticMeshComponent);
 	
 	/** Propagates per-instance vertex colors to the underlying Static Mesh for the given LOD Index */
-	static bool PropagateColorsToRawMesh(UStaticMesh* StaticMesh, int32 LODIndex, FStaticMeshComponentLODInfo& ComponentLODInfo);	
+	static UE_API bool PropagateColorsToRawMesh(UStaticMesh* StaticMesh, int32 LODIndex, FStaticMeshComponentLODInfo& ComponentLODInfo);	
 
 	/** Retrieves the Vertex Color buffer size for the given LOD level in the Static Mesh */
-	static uint32 GetVertexColorBufferSize(UMeshComponent* MeshComponent, int32 LODIndex, bool bInstance);
+	static UE_API uint32 GetVertexColorBufferSize(UMeshComponent* MeshComponent, int32 LODIndex, bool bInstance);
 
 	/** Retrieves the vertex positions from the given LOD level in the Static Mesh */
-	static TArray<FVector> GetVerticesForLOD(const UStaticMesh* StaticMesh, int32 LODIndex);
+	static UE_API TArray<FVector> GetVerticesForLOD(const UStaticMesh* StaticMesh, int32 LODIndex);
 
 	/** Retrieves the vertex colors from the given LOD level in the Static Mesh */
-	static TArray<FColor> GetColorDataForLOD(const UStaticMesh* StaticMesh, int32 LODIndex);
+	static UE_API TArray<FColor> GetColorDataForLOD(const UStaticMesh* StaticMesh, int32 LODIndex);
 
 	/** Retrieves the per-instance vertex colors from the given LOD level in the StaticMeshComponent */
-	static TArray<FColor> GetInstanceColorDataForLOD(const UStaticMeshComponent* MeshComponent, int32 LODIndex);
+	static UE_API TArray<FColor> GetInstanceColorDataForLOD(const UStaticMeshComponent* MeshComponent, int32 LODIndex);
 
 	/** Sets the specific (LOD Index) per-instance vertex colors for the given StaticMeshComponent to the supplied Color array */
-	static void SetInstanceColorDataForLOD(UStaticMeshComponent* MeshComponent, int32 LODIndex, const TArray<FColor>& Colors);	
+	static UE_API void SetInstanceColorDataForLOD(UStaticMeshComponent* MeshComponent, int32 LODIndex, const TArray<FColor>& Colors);	
 	
 	/** Sets the specific (LOD Index) per-instance vertex colors for the given StaticMeshComponent to a single Color value */
-	static void SetInstanceColorDataForLOD(UStaticMeshComponent* MeshComponent, int32 LODIndex, const FColor FillColor, const FColor MaskColor);
+	static UE_API void SetInstanceColorDataForLOD(UStaticMeshComponent* MeshComponent, int32 LODIndex, const FColor FillColor, const FColor MaskColor);
 	
 	/** Fills all vertex colors for all LODs found in the given mesh component with Fill Color */
-	static void FillStaticMeshVertexColors(UStaticMeshComponent* MeshComponent, int32 LODIndex, const FColor FillColor, const FColor MaskColor);
-	static void FillSkeletalMeshVertexColors(USkeletalMeshComponent* MeshComponent, int32 LODIndex, const FColor FillColor, const FColor MaskColor);
+	static UE_API void FillStaticMeshVertexColors(UStaticMeshComponent* MeshComponent, int32 LODIndex, const FColor FillColor, const FColor MaskColor);
+	static UE_API void FillSkeletalMeshVertexColors(USkeletalMeshComponent* MeshComponent, int32 LODIndex, const FColor FillColor, const FColor MaskColor);
 	
 	/** Sets all vertex colors for a specific LOD level in the SkeletalMesh to FillColor */
-	static void SetColorDataForLOD(USkeletalMesh* SkeletalMesh, int32 LODIndex, const FColor FillColor, const FColor MaskColor);
+	static UE_API void SetColorDataForLOD(USkeletalMesh* SkeletalMesh, int32 LODIndex, const FColor FillColor, const FColor MaskColor);
 
-	static void ApplyFillWithMask(FColor& InOutColor, const FColor& MaskColor, const FColor& FillColor);
+	static UE_API void ApplyFillWithMask(FColor& InOutColor, const FColor& MaskColor, const FColor& FillColor);
 
 	/** Helper function to import Vertex Colors from a Texture to the specified MeshComponent (makes use of SImportVertexColorsOptions Widget) */
-	static void ImportVertexColorsFromTexture(UMeshComponent* MeshComponent);
+	static UE_API void ImportVertexColorsFromTexture(UMeshComponent* MeshComponent);
 		
 	/** Checks whether or not the given Viewport Client is a VR editor viewport client */
-	static bool IsInVRMode(const FEditorViewportClient* ViewportClient);
+	static UE_API bool IsInVRMode(const FEditorViewportClient* ViewportClient);
 
 	/** Forces the Viewport Client to render using the given Viewport Color ViewMode */
-	static void SetViewportColorMode(EMeshPaintColorViewMode ColorViewMode, FEditorViewportClient* ViewportClient);
+	static UE_API void SetViewportColorMode(EMeshPaintColorViewMode ColorViewMode, FEditorViewportClient* ViewportClient);
 
 	/** Sets whether or not the level viewport should be real time rendered move or viewport as parameter? */
-	static void SetRealtimeViewport(bool bRealtime);
+	static UE_API void SetRealtimeViewport(bool bRealtime);
 
 	/** Forces the component to render LOD level at LODIndex instead of the view-based LOD level ( X = 0 means do not force the LOD, X > 0 means force the lod to X - 1 ) */
-	static void ForceRenderMeshLOD(UMeshComponent* Component, int32 LODIndex);
+	static UE_API void ForceRenderMeshLOD(UMeshComponent* Component, int32 LODIndex);
 
 	/** Clears all texture overrides for this component. */
-	static void ClearMeshTextureOverrides(const IMeshPaintGeometryAdapter& GeometryInfo, UMeshComponent* InMeshComponent);
+	static UE_API void ClearMeshTextureOverrides(const IMeshPaintGeometryAdapter& GeometryInfo, UMeshComponent* InMeshComponent);
 
 	/** Applies vertex color painting found on LOD 0 to all lower LODs. */
-	static void ApplyVertexColorsToAllLODs(IMeshPaintGeometryAdapter& GeometryInfo, UMeshComponent* InMeshComponent);
+	static UE_API void ApplyVertexColorsToAllLODs(IMeshPaintGeometryAdapter& GeometryInfo, UMeshComponent* InMeshComponent);
 
 	/** Applies the vertex colors found in LOD level 0 to all contained LOD levels in the StaticMeshComponent */
-	static void ApplyVertexColorsToAllLODs(IMeshPaintGeometryAdapter& GeometryInfo, UStaticMeshComponent* StaticMeshComponent);
+	static UE_API void ApplyVertexColorsToAllLODs(IMeshPaintGeometryAdapter& GeometryInfo, UStaticMeshComponent* StaticMeshComponent);
 
 	/** Applies the vertex colors found in LOD level 0 to all contained LOD levels in the SkeletalMeshComponent */
-	static void ApplyVertexColorsToAllLODs(IMeshPaintGeometryAdapter& GeometryInfo, USkeletalMeshComponent* SkeletalMeshComponent);
+	static UE_API void ApplyVertexColorsToAllLODs(IMeshPaintGeometryAdapter& GeometryInfo, USkeletalMeshComponent* SkeletalMeshComponent);
 
 	/** Returns the number of Mesh LODs for the given MeshComponent */
-	static int32 GetNumberOfLODs(const UMeshComponent* MeshComponent);
+	static UE_API int32 GetNumberOfLODs(const UMeshComponent* MeshComponent);
 
 	/** OutNumLODs is set to number of Mesh LODs for the given MeshComponent and returns true, or returns false of given mesh component has no valid LODs */
-	static bool TryGetNumberOfLODs(const UMeshComponent* MeshComponent, int32& OutNumLODs);
+	static UE_API bool TryGetNumberOfLODs(const UMeshComponent* MeshComponent, int32& OutNumLODs);
 	
 	/** Returns the number of Texture Coordinates for the given MeshComponent */
-	static int32 GetNumberOfUVs(const UMeshComponent* MeshComponent, int32 LODIndex);
+	static UE_API int32 GetNumberOfUVs(const UMeshComponent* MeshComponent, int32 LODIndex);
 
 	/** Checks whether or not the mesh components contains per lod colors (for all LODs)*/
-	static bool DoesMeshComponentContainPerLODColors(const UMeshComponent* MeshComponent);
+	static UE_API bool DoesMeshComponentContainPerLODColors(const UMeshComponent* MeshComponent);
 
 	/** Retrieves the number of bytes used to store the per-instance LOD vertex color data from the static mesh component */
-	static void GetInstanceColorDataInfo(const UStaticMeshComponent* StaticMeshComponent, int32 LODIndex, int32& OutTotalInstanceVertexColorBytes);
+	static UE_API void GetInstanceColorDataInfo(const UStaticMeshComponent* StaticMeshComponent, int32 LODIndex, int32& OutTotalInstanceVertexColorBytes);
 
 	/** Given arguments for an action, and an action - retrieves influences vertices and applies Action to them */
-	static bool ApplyPerVertexPaintAction(FPerVertexPaintActionArgs& InArgs, FPerVertexPaintAction Action);
+	static UE_API bool ApplyPerVertexPaintAction(FPerVertexPaintActionArgs& InArgs, FPerVertexPaintAction Action);
 
-	static bool GetPerVertexPaintInfluencedVertices(FPerVertexPaintActionArgs& InArgs, TSet<int32>& InfluencedVertices);
+	static UE_API bool GetPerVertexPaintInfluencedVertices(FPerVertexPaintActionArgs& InArgs, TSet<int32>& InfluencedVertices);
 
 	/** Given the adapter, settings and view-information retrieves influences triangles and applies Action to them */
-	static bool ApplyPerTrianglePaintAction(IMeshPaintGeometryAdapter* Adapter, const FVector& CameraPosition, const FVector& HitPosition, const UPaintBrushSettings* Settings, FPerTrianglePaintAction Action);
+	static UE_API bool ApplyPerTrianglePaintAction(IMeshPaintGeometryAdapter* Adapter, const FVector& CameraPosition, const FVector& HitPosition, const UPaintBrushSettings* Settings, FPerTrianglePaintAction Action);
 
 	/** Applies vertex painting to InOutvertexColor according to the given parameters  */
-	static bool PaintVertex(const FVector& InVertexPosition, const FMeshPaintParameters& InParams, FColor& InOutVertexColor);
+	static UE_API bool PaintVertex(const FVector& InVertexPosition, const FMeshPaintParameters& InParams, FColor& InOutVertexColor);
 
 	/** Applies Vertex Color Painting according to the given parameters */
-	static void ApplyVertexColorPaint(const FMeshPaintParameters &InParams, const FLinearColor &OldColor, FLinearColor &NewColor, const float PaintAmount);
+	static UE_API void ApplyVertexColorPaint(const FMeshPaintParameters &InParams, const FLinearColor &OldColor, FLinearColor &NewColor, const float PaintAmount);
 
 	/** Applies Vertex Blend Weight Painting according to the given parameters */
-	static void ApplyVertexWeightPaint(const FMeshPaintParameters &InParams, const FLinearColor &OldColor, const float PaintAmount, FLinearColor &NewColor);
+	static UE_API void ApplyVertexWeightPaint(const FMeshPaintParameters &InParams, const FLinearColor &OldColor, const float PaintAmount, FLinearColor &NewColor);
 
 	/** Generate texture weight color for given number of weights and the to-paint index */
-	static FLinearColor GenerateColorForTextureWeight(const int32 NumWeights, const int32 WeightIndex);
+	static UE_API FLinearColor GenerateColorForTextureWeight(const int32 NumWeights, const int32 WeightIndex);
 
 	/** Computes the Paint power multiplier value */
-	static float ComputePaintMultiplier(float SquaredDistanceToVertex2D, float BrushStrength, float BrushInnerRadius, float BrushRadialFalloff, float BrushInnerDepth, float BrushDepthFallof, float VertexDepthToBrush);
+	static UE_API float ComputePaintMultiplier(float SquaredDistanceToVertex2D, float BrushStrength, float BrushInnerRadius, float BrushRadialFalloff, float BrushInnerDepth, float BrushDepthFallof, float VertexDepthToBrush);
 
 	/** Checks whether or not a point is influenced by the painting brush according to the given parameters*/
-	static bool IsPointInfluencedByBrush(const FVector& InPosition, const FMeshPaintParameters& InParams, float& OutSquaredDistanceToVertex2D, float& OutVertexDepthToBrush);
+	static UE_API bool IsPointInfluencedByBrush(const FVector& InPosition, const FMeshPaintParameters& InParams, float& OutSquaredDistanceToVertex2D, float& OutVertexDepthToBrush);
 
-	static bool IsPointInfluencedByBrush(const FVector2D& BrushSpacePosition, const float BrushRadiusSquared, float& OutInRangeValue);
+	static UE_API bool IsPointInfluencedByBrush(const FVector2D& BrushSpacePosition, const float BrushRadiusSquared, float& OutInRangeValue);
 
 	template<typename T>
 	static void ApplyBrushToVertex(const FVector& VertexPosition, const FMatrix& InverseBrushMatrix, const float BrushRadius, const float BrushFalloffAmount, const float BrushStrength, const T& PaintValue, T& InOutValue);
@@ -162,19 +164,19 @@ public:
 		UViewportInteractor* ViewportInteractor;
 	};
 
-	static bool RetrieveViewportPaintRays(const FSceneView* View, FViewport* Viewport, FPrimitiveDrawInterface* PDI, TArray<FPaintRay>& OutPaintRays);	
+	static UE_API bool RetrieveViewportPaintRays(const FSceneView* View, FViewport* Viewport, FPrimitiveDrawInterface* PDI, TArray<FPaintRay>& OutPaintRays);	
 protected:
 	/** Imports vertex colors from a Texture to the specified Static Mesh according to user-set options */
-	static void ImportVertexColorsToStaticMesh(UStaticMesh* StaticMesh, const UVertexColorImportOptions* Options, UTexture2D* Texture);
+	static UE_API void ImportVertexColorsToStaticMesh(UStaticMesh* StaticMesh, const UVertexColorImportOptions* Options, UTexture2D* Texture);
 
 	/** Imports vertex colors from a Texture to the specified Static Mesh Component according to user-set options */
-	static void ImportVertexColorsToStaticMeshComponent(UStaticMeshComponent* StaticMeshComponent, const UVertexColorImportOptions* Options, UTexture2D* Texture);
+	static UE_API void ImportVertexColorsToStaticMeshComponent(UStaticMeshComponent* StaticMeshComponent, const UVertexColorImportOptions* Options, UTexture2D* Texture);
 
 	/** Imports vertex colors from a Texture to the specified Skeletal Mesh according to user-set options */
-	static void ImportVertexColorsToSkeletalMesh(USkeletalMesh* SkeletalMesh, const UVertexColorImportOptions* Options, UTexture2D* Texture);
+	static UE_API void ImportVertexColorsToSkeletalMesh(USkeletalMesh* SkeletalMesh, const UVertexColorImportOptions* Options, UTexture2D* Texture);
 
 	/** Helper function to retrieve vertex color from a UTexture given a UVCoordinate */
-	static FColor PickVertexColorFromTextureData(const uint8* MipData, const FVector2D& UVCoordinate, const UTexture2D* Texture, const FColor ColorMask);	
+	static UE_API FColor PickVertexColorFromTextureData(const uint8* MipData, const FVector2D& UVCoordinate, const UTexture2D* Texture, const FColor ColorMask);	
 };
 
 template<typename T>
@@ -194,3 +196,4 @@ void MeshPaintHelpers::ApplyBrushToVertex(const FVector& VertexPosition, const F
 	}	
 }
 
+#undef UE_API
diff --git a/Engine/Source/Editor/MeshPaint/Public/MeshPaintSettings.h b/Engine/Source/Editor/MeshPaint/Public/MeshPaintSettings.h
--- a/Engine/Source/Editor/MeshPaint/Public/MeshPaintSettings.h
+++ b/Engine/Source/Editor/MeshPaint/Public/MeshPaintSettings.h
@@ -6,6 +6,8 @@
 
 #include "MeshPaintSettings.generated.h"
 
+#define UE_API MESHPAINT_API
+
 /** Mesh paint color view modes (somewhat maps to EVertexColorViewMode engine enum.) */
 UENUM()
 enum class EMeshPaintColorViewMode : uint8
@@ -29,24 +31,24 @@ enum class EMeshPaintColorViewMode : uint8
 	Blue UMETA(DisplayName = "Blue Channel"),
 };
 
-UCLASS()
-class MESHPAINT_API UPaintBrushSettings : public UObject
+UCLASS(MinimalAPI)
+class UPaintBrushSettings : public UObject
 {
 	GENERATED_UCLASS_BODY()
 public:
 
-	virtual void PostEditChangeProperty(struct FPropertyChangedEvent& PropertyChangedEvent) override;
+	UE_API virtual void PostEditChangeProperty(struct FPropertyChangedEvent& PropertyChangedEvent) override;
 
-	~UPaintBrushSettings();
+	UE_API ~UPaintBrushSettings();
 
 	float GetBrushRadius() const { return BrushRadius; }
-	void SetBrushRadius(float InRadius);
+	UE_API void SetBrushRadius(float InRadius);
 
 	float GetBrushStrength() const { return BrushStrength; }
-	void SetBrushStrength(float InStrength);
+	UE_API void SetBrushStrength(float InStrength);
 
 	float GetBrushFalloff() const { return BrushFalloffAmount; }
-	void SetBrushFalloff(float InFalloff);
+	UE_API void SetBrushFalloff(float InFalloff);
 
 protected:
 	/** Radius of the Brush used for Painting */
@@ -79,8 +81,8 @@ public:
 	EMeshPaintColorViewMode ColorViewMode;
 };
 
-UCLASS(config = EditorPerProjectUserSettings)
-class MESHPAINT_API UMeshPaintSettings : public UObject
+UCLASS(MinimalAPI, config = EditorPerProjectUserSettings)
+class UMeshPaintSettings : public UObject
 {
 	GENERATED_BODY()
 
@@ -90,4 +92,6 @@ public:
 	UPROPERTY(config, EditAnywhere, Category=Visualization, Meta = (ClampMin = 0, UIMin = 0, UIMax = 10))
 	float VertexPreviewSize;
 
-};
\ No newline at end of file
+};
+
+#undef UE_API
diff --git a/Engine/Source/Editor/MeshPaint/Public/TexturePaintHelpers.h b/Engine/Source/Editor/MeshPaint/Public/TexturePaintHelpers.h
--- a/Engine/Source/Editor/MeshPaint/Public/TexturePaintHelpers.h
+++ b/Engine/Source/Editor/MeshPaint/Public/TexturePaintHelpers.h
@@ -5,13 +5,15 @@
 #include "BatchedElements.h"
 #include "MeshPaintRendering.h"
 
+#define UE_API MESHPAINT_API
+
 class UMeshComponent;
 class UTexture;
 class UTexture2D;
 class UTextureRenderTarget2D;
 
 /** Batched element parameters for texture paint shaders used for paint blending and paint mask generation */
-class MESHPAINT_API FMeshPaintBatchedElementParameters : public FBatchedElementParameters
+class FMeshPaintBatchedElementParameters : public FBatchedElementParameters
 {
 public:
 	/** Binds vertex and pixel shaders for this element */
@@ -27,7 +29,7 @@ public:
 };
 
 /** Batched element parameters for texture paint shaders used for texture dilation */
-class MESHPAINT_API FMeshPaintDilateBatchedElementParameters : public FBatchedElementParameters
+class FMeshPaintDilateBatchedElementParameters : public FBatchedElementParameters
 {
 public:
 	/** Binds vertex and pixel shaders for this element */
@@ -56,39 +58,41 @@ class IMeshPaintGeometryAdapter;
 struct FPaintableTexture;
 
 /** Helpers functions for texture painting functionality */
-class MESHPAINT_API TexturePaintHelpers
+class TexturePaintHelpers
 {
 public:
 	/** Static: Copies a texture to a render target texture */
-	static void CopyTextureToRenderTargetTexture(UTexture* SourceTexture, UTextureRenderTarget2D* RenderTargetTexture, ERHIFeatureLevel::Type FeatureLevel);
+	static UE_API void CopyTextureToRenderTargetTexture(UTexture* SourceTexture, UTextureRenderTarget2D* RenderTargetTexture, ERHIFeatureLevel::Type FeatureLevel);
 
 	/** Will generate a mask texture, used for texture dilation, and store it in the passed in render target */
-	static bool GenerateSeamMask(UMeshComponent* MeshComponent, int32 UVSet, UTextureRenderTarget2D* SeamRenderTexture, UTexture2D* Texture, UTextureRenderTarget2D* RenderTargetTexture );
+	static UE_API bool GenerateSeamMask(UMeshComponent* MeshComponent, int32 UVSet, UTextureRenderTarget2D* SeamRenderTexture, UTexture2D* Texture, UTextureRenderTarget2D* RenderTargetTexture );
 
 	/** Returns the maximum bytes per pixel that are supported for source textures when painting. This limitation is set by CreateTempUncompressedTexture. */
-	static int32 GetMaxSupportedBytesPerPixelForPainting();
+	static UE_API int32 GetMaxSupportedBytesPerPixelForPainting();
 
 	/** Returns the pixel format that CreateTempUncompressedTexture uses to create render target data for painting. */
-	static EPixelFormat GetTempUncompressedTexturePixelFormat();
+	static UE_API EPixelFormat GetTempUncompressedTexturePixelFormat();
 
 	/** Static: Creates a temporary texture used to transfer data to a render target in memory */
-	static UTexture2D* CreateTempUncompressedTexture(UTexture2D* SourceTexture);
+	static UE_API UTexture2D* CreateTempUncompressedTexture(UTexture2D* SourceTexture);
 
 	/** Makes sure that the render target is ready to paint on */
-	static void SetupInitialRenderTargetData(UTexture2D* InTextureSource, UTextureRenderTarget2D* InRenderTarget);
+	static UE_API void SetupInitialRenderTargetData(UTexture2D* InTextureSource, UTextureRenderTarget2D* InRenderTarget);
 
 	/** Tries to find Materials using the given Texture and retrieve the corresponding material indices from the MEsh Compon*/
-	static void FindMaterialIndicesUsingTexture(const UTexture* Texture, const UMeshComponent* MeshComponent, TArray<int32>& OutIndices);
+	static UE_API void FindMaterialIndicesUsingTexture(const UTexture* Texture, const UMeshComponent* MeshComponent, TArray<int32>& OutIndices);
 	
 	/** Retrieve LOD mesh sections from MeshComponent which use one of the given textures */
-	static void RetrieveMeshSectionsForTextures(const UMeshComponent* MeshComponent, int32 LODIndex, TArray<const UTexture*> Textures, TArray<FTexturePaintMeshSectionInfo>& OutSectionInfo);
+	static UE_API void RetrieveMeshSectionsForTextures(const UMeshComponent* MeshComponent, int32 LODIndex, TArray<const UTexture*> Textures, TArray<FTexturePaintMeshSectionInfo>& OutSectionInfo);
 
 	/** Retrieve LOD mesh sections from MeshComponent which contain one of the Material Indices */
-	static void RetrieveMeshSectionsForMaterialIndices(const UMeshComponent* MeshComponent, int32 LODIndex, const TArray<int32>& MaterialIndices, TArray<FTexturePaintMeshSectionInfo>& OutSectionInfo);
+	static UE_API void RetrieveMeshSectionsForMaterialIndices(const UMeshComponent* MeshComponent, int32 LODIndex, const TArray<int32>& MaterialIndices, TArray<FTexturePaintMeshSectionInfo>& OutSectionInfo);
 
 	/** Retrieves all Paintable Textures from the given MeshComponent */
-	static void RetrieveTexturesForComponent(const UMeshComponent* Component, IMeshPaintGeometryAdapter* Adapter, TArray<FPaintableTexture>& OutTextures);
+	static UE_API void RetrieveTexturesForComponent(const UMeshComponent* Component, IMeshPaintGeometryAdapter* Adapter, TArray<FPaintableTexture>& OutTextures);
 
 	/** Checks whether or not the given MeshComponent's materials use the specified texture */
-	static bool DoesMeshComponentUseTexture(UMeshComponent* MeshComponent, UTexture* Texture);
-};
\ No newline at end of file
+	static UE_API bool DoesMeshComponentUseTexture(UMeshComponent* MeshComponent, UTexture* Texture);
+};
+
+#undef UE_API
diff --git a/Engine/Source/Editor/MovieSceneCaptureDialog/Public/MovieSceneCaptureDialogModule.h b/Engine/Source/Editor/MovieSceneCaptureDialog/Public/MovieSceneCaptureDialogModule.h
--- a/Engine/Source/Editor/MovieSceneCaptureDialog/Public/MovieSceneCaptureDialogModule.h
+++ b/Engine/Source/Editor/MovieSceneCaptureDialog/Public/MovieSceneCaptureDialogModule.h
@@ -8,6 +8,8 @@
 #include "GameFramework/GameModeBase.h"
 #include "Widgets/Input/NumericTypeInterface.h"
 
+#define UE_API MOVIESCENECAPTUREDIALOG_API
+
 // Forward Declare
 class UMovieSceneCapture;
 class SNotificationItem;
@@ -57,14 +59,14 @@ struct FCaptureState
 	int32 Code;
 };
 
-struct MOVIESCENECAPTUREDIALOG_API FMovieSceneCaptureBase : TSharedFromThis<FMovieSceneCaptureBase>
+struct FMovieSceneCaptureBase : TSharedFromThis<FMovieSceneCaptureBase>
 {
 public:
 	DECLARE_MULTICAST_DELEGATE_OneParam(FCaptureStateStopped, bool);
 
 	virtual ~FMovieSceneCaptureBase() {}
 	virtual void OnCaptureStarted() {}
-	virtual void OnCaptureFinished(bool bSuccess);
+	UE_API virtual void OnCaptureFinished(bool bSuccess);
 	virtual UWorld* GetWorld() const { return nullptr; }
 	virtual void Start() {}
 	virtual void Cancel() {}
@@ -88,7 +90,7 @@ protected:
 	TFunction<void(bool)> OnFinishedCallback;
 };
 
-struct MOVIESCENECAPTUREDIALOG_API FInEditorCapture : FMovieSceneCaptureBase
+struct FInEditorCapture : FMovieSceneCaptureBase
 {
 	FInEditorCapture(UMovieSceneCapture* InCaptureObject, const TFunction<void(bool)>& InOnFinishedCallback)
 	{
@@ -103,23 +105,23 @@ struct MOVIESCENECAPTUREDIALOG_API FInEditorCapture : FMovieSceneCaptureBase
 		// Return a reference to the PIE world.
 		return CapturingFromWorld;
 	}
-	virtual void OnCaptureStarted() override;
-	virtual void Start() override;
-	virtual void Cancel() override;
-	virtual FCaptureState GetCaptureState() const override;
+	UE_API virtual void OnCaptureStarted() override;
+	UE_API virtual void Start() override;
+	UE_API virtual void Cancel() override;
+	UE_API virtual FCaptureState GetCaptureState() const override;
 	// ~FMovieSceneCaptureBase interface
 
 private:
 	/** Overrides the Level Editor Play settings to specifically disable some things (such as audio playback) */
-	void OverridePlaySettings(class ULevelEditorPlaySettings* PlayInEditorSettings);
+	UE_API void OverridePlaySettings(class ULevelEditorPlaySettings* PlayInEditorSettings);
 	/** Called when the PIE Viewport is created. */
-	void OnPIEViewportStarted();
+	UE_API void OnPIEViewportStarted();
 	/** Shuts down the capture setup , called when PIE is closed by the user or the sequence finishes playing. */
-	void Shutdown();
+	UE_API void Shutdown();
 	/** Called when the user closes the PIE instance window. */
-	void OnEndPIE(bool bIsSimulating);
+	UE_API void OnEndPIE(bool bIsSimulating);
 	/** Called when the Sequence finishes playing to the end. */
-	void OnLevelSequenceFinished();
+	UE_API void OnLevelSequenceFinished();
 
 private:
 	UWorld* CapturingFromWorld;
@@ -136,7 +138,7 @@ private:
 	TSubclassOf<AGameModeBase> CachedGameMode;
 };
 
-struct MOVIESCENECAPTUREDIALOG_API FNewProcessCapture : FMovieSceneCaptureBase
+struct FNewProcessCapture : FMovieSceneCaptureBase
 {
 
 	FNewProcessCapture(UMovieSceneCapture* InCaptureObject, const FString& InMapNameToLoad, const TFunction<void(bool)>& InOnFinishedCallback)
@@ -149,13 +151,15 @@ struct MOVIESCENECAPTUREDIALOG_API FNewProcessCapture : FMovieSceneCaptureBase
 	}
 
 	// FMovieSceneCaptureBase
-	virtual void Start() override;
-	virtual void Cancel() override;
-	virtual void OnCaptureStarted() override;
-	virtual FCaptureState GetCaptureState() const override;
+	UE_API virtual void Start() override;
+	UE_API virtual void Cancel() override;
+	UE_API virtual void OnCaptureStarted() override;
+	UE_API virtual FCaptureState GetCaptureState() const override;
 	// ~FMovieSceneCaptureBase
 
 protected:
 	TSharedPtr<FProcHandle> SharedProcHandle;
 	FString MapNameToLoad;
-};
\ No newline at end of file
+};
+
+#undef UE_API
diff --git a/Engine/Source/Editor/MovieSceneTools/Private/MVVM/ViewModels/CameraCutTrackModel.h b/Engine/Source/Editor/MovieSceneTools/Private/MVVM/ViewModels/CameraCutTrackModel.h
--- a/Engine/Source/Editor/MovieSceneTools/Private/MVVM/ViewModels/CameraCutTrackModel.h
+++ b/Engine/Source/Editor/MovieSceneTools/Private/MVVM/ViewModels/CameraCutTrackModel.h
@@ -7,6 +7,8 @@
 #include "MVVM/ViewModels/TrackModel.h"
 #include "Templates/SharedPointer.h"
 
+#define UE_API MOVIESCENETOOLS_API
+
 class UMovieSceneCameraCutTrack;
 class UMovieSceneTrack;
 namespace UE::Sequencer { template <typename T> struct TAutoRegisterViewModelTypeID; }
@@ -16,20 +18,21 @@ namespace UE
 namespace Sequencer
 {
 
-class MOVIESCENETOOLS_API FCameraCutTrackModel
+class FCameraCutTrackModel
 	: public FTrackModel
 {
 public:
 
-	UE_SEQUENCER_DECLARE_CASTABLE(FCameraCutTrackModel, FTrackModel);
+	UE_SEQUENCER_DECLARE_CASTABLE_API(UE_API, FCameraCutTrackModel, FTrackModel);
 
-	static TSharedPtr<FTrackModel> CreateTrackModel(UMovieSceneTrack* Track);
+	static UE_API TSharedPtr<FTrackModel> CreateTrackModel(UMovieSceneTrack* Track);
 
-	explicit FCameraCutTrackModel(UMovieSceneCameraCutTrack* Track);
+	UE_API explicit FCameraCutTrackModel(UMovieSceneCameraCutTrack* Track);
 
-	FSortingKey GetSortingKey() const override;
+	UE_API FSortingKey GetSortingKey() const override;
 };
 
 } // namespace Sequencer
 } // namespace UE
 
+#undef UE_API
diff --git a/Engine/Source/Editor/MovieSceneTools/Private/MVVM/ViewModels/CinematicShotTrackModel.h b/Engine/Source/Editor/MovieSceneTools/Private/MVVM/ViewModels/CinematicShotTrackModel.h
--- a/Engine/Source/Editor/MovieSceneTools/Private/MVVM/ViewModels/CinematicShotTrackModel.h
+++ b/Engine/Source/Editor/MovieSceneTools/Private/MVVM/ViewModels/CinematicShotTrackModel.h
@@ -7,6 +7,8 @@
 #include "MVVM/ViewModels/TrackModel.h"
 #include "Templates/SharedPointer.h"
 
+#define UE_API MOVIESCENETOOLS_API
+
 class UMovieSceneCinematicShotTrack;
 class UMovieSceneTrack;
 namespace UE::Sequencer { template <typename T> struct TAutoRegisterViewModelTypeID; }
@@ -16,20 +18,21 @@ namespace UE
 namespace Sequencer
 {
 
-class MOVIESCENETOOLS_API FCinematicShotTrackModel
+class FCinematicShotTrackModel
 	: public FTrackModel
 {
 public:
 
-	UE_SEQUENCER_DECLARE_CASTABLE(FCinematicShotTrackModel, FTrackModel);
+	UE_SEQUENCER_DECLARE_CASTABLE_API(UE_API, FCinematicShotTrackModel, FTrackModel);
 
-	static TSharedPtr<FTrackModel> CreateTrackModel(UMovieSceneTrack* Track);
+	static UE_API TSharedPtr<FTrackModel> CreateTrackModel(UMovieSceneTrack* Track);
 
-	explicit FCinematicShotTrackModel(UMovieSceneCinematicShotTrack* Track);
+	UE_API explicit FCinematicShotTrackModel(UMovieSceneCinematicShotTrack* Track);
 
-	FSortingKey GetSortingKey() const override;
+	UE_API FSortingKey GetSortingKey() const override;
 };
 
 } // namespace Sequencer
 } // namespace UE
 
+#undef UE_API
diff --git a/Engine/Source/Editor/MovieSceneTools/Private/MVVM/ViewModels/ScalingAnchorsModel.h b/Engine/Source/Editor/MovieSceneTools/Private/MVVM/ViewModels/ScalingAnchorsModel.h
--- a/Engine/Source/Editor/MovieSceneTools/Private/MVVM/ViewModels/ScalingAnchorsModel.h
+++ b/Engine/Source/Editor/MovieSceneTools/Private/MVVM/ViewModels/ScalingAnchorsModel.h
@@ -8,6 +8,8 @@
 #include "UObject/WeakObjectPtr.h"
 #include "UObject/WeakObjectPtrTemplates.h"
 
+#define UE_API MOVIESCENETOOLS_API
+
 class ISequencer;
 class UMovieSceneScalingAnchors;
 class UToolMenu;
@@ -17,28 +19,28 @@ namespace UE::Sequencer
 
 class FTrackAreaViewModel;
 
-class MOVIESCENETOOLS_API FScalingAnchorsModel
+class FScalingAnchorsModel
 	: public FViewModel
 	, public IObjectModelExtension
 	, public ITrackAreaViewSpaceProviderExtension
 {
 public:
 
-	UE_SEQUENCER_DECLARE_CASTABLE(FScalingAnchorsModel, FViewModel, IObjectModelExtension, ITrackAreaViewSpaceProviderExtension);
+	UE_SEQUENCER_DECLARE_CASTABLE_API(UE_API, FScalingAnchorsModel, FViewModel, IObjectModelExtension, ITrackAreaViewSpaceProviderExtension);
 
 	/*~ Begin IObjectModelExtension */
-	virtual void InitializeObject(TWeakObjectPtr<> InWeakObject) override;
-	virtual UObject* GetObject() const override;
+	UE_API virtual void InitializeObject(TWeakObjectPtr<> InWeakObject) override;
+	UE_API virtual UObject* GetObject() const override;
 	/*~ End IObjectModelExtension */
 
 	/*~ Begin ITrackAreaViewSpaceProviderExtension */
-	virtual void UpdateViewSpaces(FTrackAreaViewModel& TrackAreaViewModel) override;
+	UE_API virtual void UpdateViewSpaces(FTrackAreaViewModel& TrackAreaViewModel) override;
 	/*~ End ITrackAreaViewSpaceProviderExtension */
 
 private:
 
-	void ExtendSectionMenu(UToolMenu* InMenu);
-	void CreateScalingGroup(TWeakPtr<ISequencer> InWeakSequencer);
+	UE_API void ExtendSectionMenu(UToolMenu* InMenu);
+	UE_API void CreateScalingGroup(TWeakPtr<ISequencer> InWeakSequencer);
 
 private:
 
@@ -48,3 +50,4 @@ private:
 
 } // namespace UE::Sequencer
 
+#undef UE_API
diff --git a/Engine/Source/Editor/MovieSceneTools/Private/Sections/MovieSceneSectionHelpers.h b/Engine/Source/Editor/MovieSceneTools/Private/Sections/MovieSceneSectionHelpers.h
--- a/Engine/Source/Editor/MovieSceneTools/Private/Sections/MovieSceneSectionHelpers.h
+++ b/Engine/Source/Editor/MovieSceneTools/Private/Sections/MovieSceneSectionHelpers.h
@@ -5,34 +5,38 @@
 #include "Widgets/SWindow.h"
 #include "Containers/ArrayView.h"
 
+#define UE_API MOVIESCENETOOLS_API
+
 struct FKeyHandle;
 struct FTimeToPixel;
 struct FMovieSceneFloatChannel;
 class ISequencer;
 class UMovieSceneSection;
 
-class MOVIESCENETOOLS_API MovieSceneSectionHelpers
+class MovieSceneSectionHelpers
 {
 public:
 
 	/** Consolidate color curves for all track sections. */
-	static void ConsolidateColorCurves(TArray< TTuple<float, FLinearColor> >& OutColorKeys, const FLinearColor& DefaultColor, TArrayView<const FMovieSceneFloatChannel* const> ColorChannels, const FTimeToPixel& TimeConverter);
+	static UE_API void ConsolidateColorCurves(TArray< TTuple<float, FLinearColor> >& OutColorKeys, const FLinearColor& DefaultColor, TArrayView<const FMovieSceneFloatChannel* const> ColorChannels, const FTimeToPixel& TimeConverter);
 };
 
-class MOVIESCENETOOLS_API FMovieSceneKeyColorPicker
+class FMovieSceneKeyColorPicker
 {
 public:
-	FMovieSceneKeyColorPicker(UMovieSceneSection* Section, FMovieSceneFloatChannel* RChannel, FMovieSceneFloatChannel* GChannel, FMovieSceneFloatChannel* BChannel, FMovieSceneFloatChannel* AChannel, const TArray<FKeyHandle>& KeyHandles, TWeakPtr<ISequencer> InSequencer);
+	UE_API FMovieSceneKeyColorPicker(UMovieSceneSection* Section, FMovieSceneFloatChannel* RChannel, FMovieSceneFloatChannel* GChannel, FMovieSceneFloatChannel* BChannel, FMovieSceneFloatChannel* AChannel, const TArray<FKeyHandle>& KeyHandles, TWeakPtr<ISequencer> InSequencer);
 
 private:
 		
-	void OnColorPickerPicked(FLinearColor NewColor, UMovieSceneSection* Section, FMovieSceneFloatChannel* RChannel, FMovieSceneFloatChannel* GChannel, FMovieSceneFloatChannel* BChannel, FMovieSceneFloatChannel* AChannel, TWeakPtr<ISequencer> InSequencer);
-	void OnColorPickerClosed(const TSharedRef<SWindow>& Window, UMovieSceneSection* Section, FMovieSceneFloatChannel* RChannel, FMovieSceneFloatChannel* GChannel, FMovieSceneFloatChannel* BChannel, FMovieSceneFloatChannel* AChannel, TWeakPtr<ISequencer> InSequencer);
-	void OnColorPickerCancelled(FLinearColor NewColor, UMovieSceneSection* Section, FMovieSceneFloatChannel* RChannel, FMovieSceneFloatChannel* GChannel, FMovieSceneFloatChannel* BChannel, FMovieSceneFloatChannel* AChannel, TWeakPtr<ISequencer> InSequencer);
+	UE_API void OnColorPickerPicked(FLinearColor NewColor, UMovieSceneSection* Section, FMovieSceneFloatChannel* RChannel, FMovieSceneFloatChannel* GChannel, FMovieSceneFloatChannel* BChannel, FMovieSceneFloatChannel* AChannel, TWeakPtr<ISequencer> InSequencer);
+	UE_API void OnColorPickerClosed(const TSharedRef<SWindow>& Window, UMovieSceneSection* Section, FMovieSceneFloatChannel* RChannel, FMovieSceneFloatChannel* GChannel, FMovieSceneFloatChannel* BChannel, FMovieSceneFloatChannel* AChannel, TWeakPtr<ISequencer> InSequencer);
+	UE_API void OnColorPickerCancelled(FLinearColor NewColor, UMovieSceneSection* Section, FMovieSceneFloatChannel* RChannel, FMovieSceneFloatChannel* GChannel, FMovieSceneFloatChannel* BChannel, FMovieSceneFloatChannel* AChannel, TWeakPtr<ISequencer> InSequencer);
 
 private:
 
-	static FFrameNumber KeyTime;
-	static FLinearColor InitialColor;
-	static bool bColorPickerWasCancelled;
-};
\ No newline at end of file
+	static UE_API FFrameNumber KeyTime;
+	static UE_API FLinearColor InitialColor;
+	static UE_API bool bColorPickerWasCancelled;
+};
+
+#undef UE_API
diff --git a/Engine/Source/Editor/MovieSceneTools/Private/TrackEditors/MaterialParameterCollectionTrackEditor.h b/Engine/Source/Editor/MovieSceneTools/Private/TrackEditors/MaterialParameterCollectionTrackEditor.h
--- a/Engine/Source/Editor/MovieSceneTools/Private/TrackEditors/MaterialParameterCollectionTrackEditor.h
+++ b/Engine/Source/Editor/MovieSceneTools/Private/TrackEditors/MaterialParameterCollectionTrackEditor.h
@@ -7,6 +7,8 @@
 #include "Templates/SharedPointer.h"
 #include "Templates/SubclassOf.h"
 
+#define UE_API MOVIESCENETOOLS_API
+
 class FMenuBuilder;
 class ISequencer;
 class ISequencerSection;
@@ -28,38 +30,40 @@ struct FSlateBrush;
 /**
  * Track editor for material parameter collections.
  */
-class MOVIESCENETOOLS_API FMaterialParameterCollectionTrackEditor
+class FMaterialParameterCollectionTrackEditor
 	: public FMovieSceneTrackEditor
 {
 public:
 
 	/** Constructor. */
-	FMaterialParameterCollectionTrackEditor(TSharedRef<ISequencer> InSequencer);
+	UE_API FMaterialParameterCollectionTrackEditor(TSharedRef<ISequencer> InSequencer);
 
 	/** Factory function */
-	static TSharedRef<ISequencerTrackEditor> CreateTrackEditor(TSharedRef<ISequencer> OwningSequencer);
+	static UE_API TSharedRef<ISequencerTrackEditor> CreateTrackEditor(TSharedRef<ISequencer> OwningSequencer);
 
 public:
 
 	//~ ISequencerTrackEditor interface
-	virtual FText GetDisplayName() const override;
-	virtual TSharedPtr<SWidget> BuildOutlinerEditWidget(const FGuid& ObjectBinding, UMovieSceneTrack* Track, const FBuildEditWidgetParams& Params) override;
-	virtual TSharedPtr<SWidget> BuildOutlinerEditColumnWidget(const FGuid& ObjectBinding, UMovieSceneTrack* Track, const FBuildEditWidgetParams& Params) override;
-	virtual TSharedRef<ISequencerSection> MakeSectionInterface(UMovieSceneSection& SectionObject, UMovieSceneTrack& Track, FGuid ObjectBinding) override;
-	virtual bool SupportsType(TSubclassOf<UMovieSceneTrack> Type) const override;
-	virtual bool SupportsSequence(UMovieSceneSequence* InSequence) const override;
-	virtual void BuildAddTrackMenu(FMenuBuilder& MenuBuilder) override;
-	virtual const FSlateBrush* GetIconBrush() const override;
-	virtual bool HandleAssetAdded(UObject* Asset, const FGuid& TargetObjectGuid) override;
+	UE_API virtual FText GetDisplayName() const override;
+	UE_API virtual TSharedPtr<SWidget> BuildOutlinerEditWidget(const FGuid& ObjectBinding, UMovieSceneTrack* Track, const FBuildEditWidgetParams& Params) override;
+	UE_API virtual TSharedPtr<SWidget> BuildOutlinerEditColumnWidget(const FGuid& ObjectBinding, UMovieSceneTrack* Track, const FBuildEditWidgetParams& Params) override;
+	UE_API virtual TSharedRef<ISequencerSection> MakeSectionInterface(UMovieSceneSection& SectionObject, UMovieSceneTrack& Track, FGuid ObjectBinding) override;
+	UE_API virtual bool SupportsType(TSubclassOf<UMovieSceneTrack> Type) const override;
+	UE_API virtual bool SupportsSequence(UMovieSceneSequence* InSequence) const override;
+	UE_API virtual void BuildAddTrackMenu(FMenuBuilder& MenuBuilder) override;
+	UE_API virtual const FSlateBrush* GetIconBrush() const override;
+	UE_API virtual bool HandleAssetAdded(UObject* Asset, const FGuid& TargetObjectGuid) override;
 
 private:
 
 	/** Provides the contents of the add parameter menu. */
-	TSharedRef<SWidget> OnGetAddParameterMenuContent(UMovieSceneMaterialParameterCollectionTrack* MaterialTrack, int32 RowIndex, int32 TrackInsertRowIndex);
+	UE_API TSharedRef<SWidget> OnGetAddParameterMenuContent(UMovieSceneMaterialParameterCollectionTrack* MaterialTrack, int32 RowIndex, int32 TrackInsertRowIndex);
 
-	void AddScalarParameter(UMovieSceneMaterialParameterCollectionTrack* Track, int32 RowIndex, FCollectionScalarParameter Parameter);
-	void AddVectorParameter(UMovieSceneMaterialParameterCollectionTrack* Track, int32 RowIndex, FCollectionVectorParameter Parameter);
+	UE_API void AddScalarParameter(UMovieSceneMaterialParameterCollectionTrack* Track, int32 RowIndex, FCollectionScalarParameter Parameter);
+	UE_API void AddVectorParameter(UMovieSceneMaterialParameterCollectionTrack* Track, int32 RowIndex, FCollectionVectorParameter Parameter);
 
-	void AddTrackToSequence(const FAssetData& InAssetData);
-	void AddTrackToSequenceEnterPressed(const TArray<FAssetData>& InAssetData);
+	UE_API void AddTrackToSequence(const FAssetData& InAssetData);
+	UE_API void AddTrackToSequenceEnterPressed(const TArray<FAssetData>& InAssetData);
 };
+
+#undef UE_API
diff --git a/Engine/Source/Editor/MovieSceneTools/Public/AutomatedLevelSequenceCapture.h b/Engine/Source/Editor/MovieSceneTools/Public/AutomatedLevelSequenceCapture.h
--- a/Engine/Source/Editor/MovieSceneTools/Public/AutomatedLevelSequenceCapture.h
+++ b/Engine/Source/Editor/MovieSceneTools/Public/AutomatedLevelSequenceCapture.h
@@ -9,22 +9,24 @@
 #include "MovieSceneCapture.h"
 #include "AutomatedLevelSequenceCapture.generated.h"
 
+#define UE_API MOVIESCENETOOLS_API
+
 class ALevelSequenceActor;
 class FJsonObject;
 class FSceneViewport;
 class ULevelSequenceBurnInOptions;
 struct FMovieSceneTimeController_FrameStep;
 
-UCLASS(config=EditorPerProjectUserSettings, PerObjectConfig, BlueprintType)
-class MOVIESCENETOOLS_API UAutomatedLevelSequenceCapture : public UMovieSceneCapture
+UCLASS(MinimalAPI, config=EditorPerProjectUserSettings, PerObjectConfig, BlueprintType)
+class UAutomatedLevelSequenceCapture : public UMovieSceneCapture
 {
 public:
-	UAutomatedLevelSequenceCapture(const FObjectInitializer&);
+	UE_API UAutomatedLevelSequenceCapture(const FObjectInitializer&);
 
 	GENERATED_BODY()
 
 	/** This name is used by the UI to save/load a specific instance of the settings from config that doesn't affect the CDO which would affect scripting environments. */
-	static const FName AutomatedLevelSequenceCaptureUIName;
+	static UE_API const FName AutomatedLevelSequenceCaptureUIName;
 
 	/** A level sequence asset to playback at runtime - used where the level sequence does not already exist in the world. */
 	UPROPERTY(BlueprintReadWrite, Category=Animation)
@@ -80,71 +82,71 @@ public:
 
 public:
 	// UMovieSceneCapture interface
-	virtual void Initialize(TSharedPtr<FSceneViewport> InViewport, int32 PIEInstance = -1) override;
+	UE_API virtual void Initialize(TSharedPtr<FSceneViewport> InViewport, int32 PIEInstance = -1) override;
 
-	virtual void LoadFromConfig() override;
+	UE_API virtual void LoadFromConfig() override;
 
-	virtual void SaveToConfig() override;
+	UE_API virtual void SaveToConfig() override;
 
-	virtual void Close() override;
+	UE_API virtual void Close() override;
 
-	double GetEstimatedCaptureDurationSeconds() const override;
+	UE_API double GetEstimatedCaptureDurationSeconds() const override;
 	
 	/** Override the render frames with the given start/end frames. Restore the values when done rendering. */
-	void SetFrameOverrides(FFrameNumber InStartFrame, FFrameNumber InEndFrame);
+	UE_API void SetFrameOverrides(FFrameNumber InStartFrame, FFrameNumber InEndFrame);
 
 protected:
 
-	virtual void AddFormatMappings(TMap<FString, FStringFormatArg>& OutFormatMappings, const FFrameMetrics& FrameMetrics) const override;
+	UE_API virtual void AddFormatMappings(TMap<FString, FStringFormatArg>& OutFormatMappings, const FFrameMetrics& FrameMetrics) const override;
 
 	/** Custom, additional json serialization */
-	virtual void SerializeAdditionalJson(FJsonObject& Object);
+	UE_API virtual void SerializeAdditionalJson(FJsonObject& Object);
 
 	/** Custom, additional json deserialization */
-	virtual void DeserializeAdditionalJson(const FJsonObject& Object);
+	UE_API virtual void DeserializeAdditionalJson(const FJsonObject& Object);
 
 private:
 
 	/** Update any cached information we need from the level sequence actor */
-	void UpdateFrameState();
+	UE_API void UpdateFrameState();
 
 	/** Called when the level sequence has updated the world */
-	void SequenceUpdated(const UMovieSceneSequencePlayer& Player, FFrameTime CurrentTime, FFrameTime PreviousTime);
+	UE_API void SequenceUpdated(const UMovieSceneSequencePlayer& Player, FFrameTime CurrentTime, FFrameTime PreviousTime);
 
 	/** Called to set up the player's playback range */
-	void SetupFrameRange();
+	UE_API void SetupFrameRange();
 
 	/** Enable cinematic mode override */
-	void EnableCinematicMode();
+	UE_API void EnableCinematicMode();
 
 	/** Export EDL if requested */
-	void ExportEDL();
+	UE_API void ExportEDL();
 
 	/** Export FCPXML if requested */
-	void ExportFCPXML();
+	UE_API void ExportFCPXML();
 
 	/** Delegate binding for the above callback */
 	FDelegateHandle OnPlayerUpdatedBinding;
 
 private:
 
-	virtual void OnTick(float DeltaSeconds) override;
+	UE_API virtual void OnTick(float DeltaSeconds) override;
 
 	/** Initialize all the shots to be recorded, ie. expand section ranges with handle frames */
-	bool InitializeShots();
+	UE_API bool InitializeShots();
 
 	/** Set up the current shot to be recorded, ie. expand playback range to the section range */
-	bool SetupShot(FFrameNumber& StartTime, FFrameNumber& EndTime);
+	UE_API bool SetupShot(FFrameNumber& StartTime, FFrameNumber& EndTime);
 
 	/** Restore any modification to shots */
-	void RestoreShots();
+	UE_API void RestoreShots();
 
 	/** Restore frame settings from overridden shot frames */
-	bool RestoreFrameOverrides();
+	UE_API bool RestoreFrameOverrides();
 
-	void DelayBeforeWarmupFinished();
+	UE_API void DelayBeforeWarmupFinished();
 
-	void PauseFinished();
+	UE_API void PauseFinished();
 
 
 	/** The pre-existing level sequence actor to use for capture that specifies playback settings */
@@ -211,3 +213,4 @@ private:
 	int32 CachedDelayBeforeShotWarmUp;
 };
 
+#undef UE_API
diff --git a/Engine/Source/Editor/MovieSceneTools/Public/BakingAnimationKeySettings.h b/Engine/Source/Editor/MovieSceneTools/Public/BakingAnimationKeySettings.h
--- a/Engine/Source/Editor/MovieSceneTools/Public/BakingAnimationKeySettings.h
+++ b/Engine/Source/Editor/MovieSceneTools/Public/BakingAnimationKeySettings.h
@@ -13,7 +13,7 @@ enum class EBakingKeySettings : uint8
 };
 
 USTRUCT(BlueprintType)
-struct MOVIESCENETOOLS_API FBakingAnimationKeySettings
+struct FBakingAnimationKeySettings
 {
 	GENERATED_BODY();
 
diff --git a/Engine/Source/Editor/MovieSceneTools/Public/Channels/SCurveEditorEventChannelView.h b/Engine/Source/Editor/MovieSceneTools/Public/Channels/SCurveEditorEventChannelView.h
--- a/Engine/Source/Editor/MovieSceneTools/Public/Channels/SCurveEditorEventChannelView.h
+++ b/Engine/Source/Editor/MovieSceneTools/Public/Channels/SCurveEditorEventChannelView.h
@@ -9,6 +9,8 @@
 #include "Templates/SharedPointer.h"
 #include "Views/SInteractiveCurveEditorView.h"
 
+#define UE_API MOVIESCENETOOLS_API
+
 class FCurveEditor;
 class FPaintArgs;
 class FSlateRect;
@@ -17,25 +19,27 @@ class FText;
 class FWidgetStyle;
 struct FGeometry;
 
-class MOVIESCENETOOLS_API SCurveEditorEventChannelView : public SInteractiveCurveEditorView
+class SCurveEditorEventChannelView : public SInteractiveCurveEditorView
 {
 public:
 
-	void Construct(const FArguments& InArgs, TWeakPtr<FCurveEditor> InCurveEditor);
+	UE_API void Construct(const FArguments& InArgs, TWeakPtr<FCurveEditor> InCurveEditor);
 
-	virtual void GetGridLinesY(TSharedRef<const FCurveEditor> CurveEditor, TArray<float>& MajorGridLines, TArray<float>& MinorGridLines, TArray<FText>* MajorGridLabels) const override;
+	UE_API virtual void GetGridLinesY(TSharedRef<const FCurveEditor> CurveEditor, TArray<float>& MajorGridLines, TArray<float>& MinorGridLines, TArray<FText>* MajorGridLabels) const override;
 
 private:
 
 	//~ SWidget Interface
-	virtual void Tick(const FGeometry& AllottedGeometry, const double InCurrentTime, const float InDeltaTime) override;
-	virtual FVector2D ComputeDesiredSize(float LayoutScaleMultiplier) const override;
+	UE_API virtual void Tick(const FGeometry& AllottedGeometry, const double InCurrentTime, const float InDeltaTime) override;
+	UE_API virtual FVector2D ComputeDesiredSize(float LayoutScaleMultiplier) const override;
 
 	//~ SInteractiveCurveEditorView interface
-	virtual void PaintView(const FPaintArgs& Args, const FGeometry& AllottedGeometry, const FSlateRect& MyCullingRect, FSlateWindowElementList& OutDrawElements, int32 BaseLayerId, const FWidgetStyle& InWidgetStyle, bool bParentEnabled) const override;
+	UE_API virtual void PaintView(const FPaintArgs& Args, const FGeometry& AllottedGeometry, const FSlateRect& MyCullingRect, FSlateWindowElementList& OutDrawElements, int32 BaseLayerId, const FWidgetStyle& InWidgetStyle, bool bParentEnabled) const override;
 	virtual bool IsValueSnapEnabled() const override { return false; }
 
-	void DrawLabels(const FGeometry& AllottedGeometry, FSlateWindowElementList& OutDrawElements, int32 BaseLayerId, ESlateDrawEffect DrawEffects) const;
+	UE_API void DrawLabels(const FGeometry& AllottedGeometry, FSlateWindowElementList& OutDrawElements, int32 BaseLayerId, ESlateDrawEffect DrawEffects) const;
 
-	static float TrackHeight;
+	static UE_API float TrackHeight;
 };
+
+#undef UE_API
